HierarchyTable Built-ins

diff --git a/Engine/Plugins/Experimental/Animation/HierarchyTableBuiltin/HierarchyTableBuiltin.uplugin b/Engine/Plugins/Experimental/Animation/HierarchyTableBuiltin/HierarchyTableBuiltin.uplugin
new file mode 100644
--- /dev/null
+++ b/Engine/Plugins/Experimental/Animation/HierarchyTableBuiltin/HierarchyTableBuiltin.uplugin
@@ -0,0 +1,40 @@
+{
+	"FileVersion": 3,
+	"Version": 1,
+	"VersionName": "0.1",
+	"FriendlyName": "Hierarchy Table Builtin",
+	"Description": "Contains built-in types for controlling animation blending",
+	"Category": "Animation",
+	"CreatedBy": "Epic Games, Inc.",
+	"CreatedByURL": "https://epicgames.com",
+	"DocsURL": "",
+	"MarketplaceURL": "",
+	"SupportURL": "",
+	"EnabledByDefault": false,
+	"CanContainContent": true,
+	"IsBetaVersion": false,
+	"IsExperimentalVersion": true,
+	"Installed": false,
+	"Modules": [
+		{
+			"Name": "HierarchyTableBuiltinRuntime",
+			"Type": "Runtime",
+			"LoadingPhase": "Default"
+		},
+		{
+			"Name": "HierarchyTableBuiltinEditor",
+			"Type": "Editor",
+			"LoadingPhase": "Default"
+		}
+	],
+	"Plugins": [
+		{
+			"Name": "HierarchyTable",
+			"Enabled": true
+		},
+		{
+			"Name": "TypedElementDataStorageEditorUI",
+			"Enabled": true
+		}
+	]
+}
diff --git a/Engine/Plugins/Experimental/Animation/HierarchyTableBuiltin/Source/Editor/HierarchyTableBuiltinEditor.Build.cs b/Engine/Plugins/Experimental/Animation/HierarchyTableBuiltin/Source/Editor/HierarchyTableBuiltinEditor.Build.cs
new file mode 100644
--- /dev/null
+++ b/Engine/Plugins/Experimental/Animation/HierarchyTableBuiltin/Source/Editor/HierarchyTableBuiltinEditor.Build.cs
@@ -0,0 +1,24 @@
+// Copyright Epic Games, Inc. All Rights Reserved.
+
+using EpicGames.Core;
+using UnrealBuildTool;
+
+public class HierarchyTableBuiltinEditor : ModuleRules
+{
+	public HierarchyTableBuiltinEditor(ReadOnlyTargetRules Target) : base(Target)
+	{
+		PrivateDependencyModuleNames.AddAll(
+			"Core",
+			"CoreUObject",
+			"Engine", 
+			"InputCore",
+			"Slate",
+			"SlateCore",
+			"UnrealEd",
+			"HierarchyTableRuntime",
+			"HierarchyTableEditor",
+			"HierarchyTableBuiltinRuntime",
+			"TypedElementFramework" // TODO: Remove
+		);
+	}
+}
diff --git a/Engine/Plugins/Experimental/Animation/HierarchyTableBuiltin/Source/Editor/Private/HierachyTableBuiltinEditorModule.cpp b/Engine/Plugins/Experimental/Animation/HierarchyTableBuiltin/Source/Editor/Private/HierachyTableBuiltinEditorModule.cpp
new file mode 100644
--- /dev/null
+++ b/Engine/Plugins/Experimental/Animation/HierarchyTableBuiltin/Source/Editor/Private/HierachyTableBuiltinEditorModule.cpp
@@ -0,0 +1,25 @@
+// Copyright Epic Games, Inc. All Rights Reserved.
+
+#include "HierarchyTableBuiltinEditorModule.h"
+#include "HierarchyTableTypeRegistry.h"
+#include "TimeProfile/HierarchyTableTypeTime.h"
+#include "TimeProfile/TimeProfileTypeHandler.h"
+#include "MaskProfile/HierarchyTableTypeMask.h"
+#include "MaskProfile/MaskProfileTypeHandler.h"
+#include "Modules/ModuleManager.h"
+
+void FHierarchyTableBuiltinEditorModule::StartupModule()
+{
+	UHierarchyTableTypeRegistry* Registry = GetMutableDefault<UHierarchyTableTypeRegistry>();
+	Registry->Register(FHierarchyTableType_Mask::StaticStruct(), GetDefault<UHierarchyTableTypeHandler_Mask>());
+	Registry->Register(FHierarchyTableType_Time::StaticStruct(), GetDefault<UHierarchyTableTypeHandler_Time>());
+}
+
+void FHierarchyTableBuiltinEditorModule::ShutdownModule()
+{
+	UHierarchyTableTypeRegistry* Registry = GetMutableDefault<UHierarchyTableTypeRegistry>();
+	Registry->Unregister(FHierarchyTableType_Mask::StaticStruct());
+	Registry->Unregister(FHierarchyTableType_Time::StaticStruct());
+}
+
+IMPLEMENT_MODULE(FHierarchyTableBuiltinEditorModule, HierarchyTableBuiltinEditor)
diff --git a/Engine/Plugins/Experimental/Animation/HierarchyTableBuiltin/Source/Editor/Private/MaskProfile/MaskProfileFactory.cpp b/Engine/Plugins/Experimental/Animation/HierarchyTableBuiltin/Source/Editor/Private/MaskProfile/MaskProfileFactory.cpp
new file mode 100644
--- /dev/null
+++ b/Engine/Plugins/Experimental/Animation/HierarchyTableBuiltin/Source/Editor/Private/MaskProfile/MaskProfileFactory.cpp
@@ -0,0 +1,12 @@
+// Copyright Epic Games, Inc. All Rights Reserved.
+
+#include "MaskProfile/MaskProfileFactory.h"
+#include "MaskProfile/MaskProfileProxyColumn.h"
+#include "MaskProfile/MaskProfileWidgetConstructor.h"
+
+void UHierarchyTableMaskFactory::RegisterWidgetConstructors(ITypedElementDataStorageInterface& DataStorage, ITypedElementDataStorageUiInterface& DataStorageUi) const
+{
+	using namespace TypedElementDataStorage;
+
+	DataStorageUi.RegisterWidgetFactory<FHierarchyTableMaskWidgetConstructor_Value>(FName(TEXT("General.Cell")), FColumn<FHierarchyTableMaskColumn_Value>());
+}
\ No newline at end of file
diff --git a/Engine/Plugins/Experimental/Animation/HierarchyTableBuiltin/Source/Editor/Private/MaskProfile/MaskProfileFactory.h b/Engine/Plugins/Experimental/Animation/HierarchyTableBuiltin/Source/Editor/Private/MaskProfile/MaskProfileFactory.h
new file mode 100644
--- /dev/null
+++ b/Engine/Plugins/Experimental/Animation/HierarchyTableBuiltin/Source/Editor/Private/MaskProfile/MaskProfileFactory.h
@@ -0,0 +1,16 @@
+// Copyright Epic Games, Inc. All Rights Reserved.
+
+#pragma once
+
+#include "Elements/Interfaces/TypedElementDataStorageFactory.h"
+
+#include "MaskProfileFactory.generated.h"
+
+UCLASS()
+class UHierarchyTableMaskFactory : public UTypedElementDataStorageFactory
+{
+	GENERATED_BODY()
+
+public:
+	void RegisterWidgetConstructors(ITypedElementDataStorageInterface& DataStorage, ITypedElementDataStorageUiInterface& DataStorageUi) const override;
+};
\ No newline at end of file
diff --git a/Engine/Plugins/Experimental/Animation/HierarchyTableBuiltin/Source/Editor/Private/MaskProfile/MaskProfileTypeHandler.cpp b/Engine/Plugins/Experimental/Animation/HierarchyTableBuiltin/Source/Editor/Private/MaskProfile/MaskProfileTypeHandler.cpp
new file mode 100644
--- /dev/null
+++ b/Engine/Plugins/Experimental/Animation/HierarchyTableBuiltin/Source/Editor/Private/MaskProfile/MaskProfileTypeHandler.cpp
@@ -0,0 +1,22 @@
+// Copyright Epic Games, Inc. All Rights Reserved.
+
+#pragma once
+
+#include "MaskProfile/MaskProfileTypeHandler.h"
+#include "MaskProfile/HierarchyTableTypeMask.h"
+#include "MaskProfile/MaskProfileProxyColumn.h"
+
+FInstancedStruct UHierarchyTableTypeHandler_Mask::GetDefaultEntry() const
+{
+	FInstancedStruct NewEntry;
+	NewEntry.InitializeAs(FHierarchyTableType_Mask::StaticStruct());
+	FHierarchyTableType_Mask& NewEntryRef = NewEntry.GetMutable<FHierarchyTableType_Mask>();
+	NewEntryRef.Value = 1.0f;
+
+	return NewEntry;
+}
+
+TArray<UScriptStruct*> UHierarchyTableTypeHandler_Mask::GetColumns() const
+{
+	return { FHierarchyTableMaskColumn_Value::StaticStruct() };
+}
\ No newline at end of file
diff --git a/Engine/Plugins/Experimental/Animation/HierarchyTableBuiltin/Source/Editor/Private/MaskProfile/MaskProfileTypeHandler.h b/Engine/Plugins/Experimental/Animation/HierarchyTableBuiltin/Source/Editor/Private/MaskProfile/MaskProfileTypeHandler.h
new file mode 100644
--- /dev/null
+++ b/Engine/Plugins/Experimental/Animation/HierarchyTableBuiltin/Source/Editor/Private/MaskProfile/MaskProfileTypeHandler.h
@@ -0,0 +1,18 @@
+// Copyright Epic Games, Inc. All Rights Reserved.
+
+#pragma once
+
+#include "HierarchyTableTypeRegistry.h"
+
+#include "MaskProfileTypeHandler.generated.h"
+
+UCLASS()
+class UHierarchyTableTypeHandler_Mask final : public UHierarchyTableTypeHandler_Base
+{
+	GENERATED_BODY()
+
+public:
+	FInstancedStruct GetDefaultEntry() const override;
+
+	TArray<UScriptStruct*> GetColumns() const override;
+};
\ No newline at end of file
diff --git a/Engine/Plugins/Experimental/Animation/HierarchyTableBuiltin/Source/Editor/Private/MaskProfile/MaskProfileWidgetConstructor.cpp b/Engine/Plugins/Experimental/Animation/HierarchyTableBuiltin/Source/Editor/Private/MaskProfile/MaskProfileWidgetConstructor.cpp
new file mode 100644
--- /dev/null
+++ b/Engine/Plugins/Experimental/Animation/HierarchyTableBuiltin/Source/Editor/Private/MaskProfile/MaskProfileWidgetConstructor.cpp
@@ -0,0 +1,23 @@
+// Copyright Epic Games, Inc. All Rights Reserved.
+
+#include "MaskProfile/MaskProfileWidgetConstructor.h"
+#include "Widgets/Input/SSpinBox.h"
+#include "MaskProfile/HierarchyTableTypeMask.h"
+#include "HierarchyTable.h"
+
+TSharedRef<SWidget> FHierarchyTableMaskWidgetConstructor_Value::CreateWidget(FHierarchyTableEntryData* EntryData)
+{
+	return SNew(SSpinBox<float>)
+		.IsEnabled_Lambda([EntryData]() { return EntryData->IsOverridden(); })
+		.MinDesiredWidth(100.0f)
+		.MinValue(0.0f)
+		.MaxValue(1.0f)
+		.Value_Lambda([EntryData]()
+			{
+				return EntryData->GetValue<FHierarchyTableType_Mask>()->Value;
+			})
+		.OnValueChanged_Lambda([EntryData](float NewValue)
+			{
+				EntryData->GetMutableValue<FHierarchyTableType_Mask>()->Value = NewValue;
+			});
+}
\ No newline at end of file
diff --git a/Engine/Plugins/Experimental/Animation/HierarchyTableBuiltin/Source/Editor/Private/MaskProfile/MaskProfileWidgetConstructor.h b/Engine/Plugins/Experimental/Animation/HierarchyTableBuiltin/Source/Editor/Private/MaskProfile/MaskProfileWidgetConstructor.h
new file mode 100644
--- /dev/null
+++ b/Engine/Plugins/Experimental/Animation/HierarchyTableBuiltin/Source/Editor/Private/MaskProfile/MaskProfileWidgetConstructor.h
@@ -0,0 +1,21 @@
+// Copyright Epic Games, Inc. All Rights Reserved.
+
+#pragma once
+
+#include "Elements/Interfaces/TypedElementDataStorageUiInterface.h"
+#include "HierarchyTableWidgetConstructor.h"
+
+#include "MaskProfileWidgetConstructor.generated.h"
+
+USTRUCT()
+struct FHierarchyTableMaskWidgetConstructor_Value : public FHierarchyTableWidgetConstructor
+{
+	GENERATED_BODY()
+
+public:
+	FHierarchyTableMaskWidgetConstructor_Value() : Super(StaticStruct()) {}
+	~FHierarchyTableMaskWidgetConstructor_Value() override = default;
+
+protected:
+	virtual TSharedRef<SWidget> CreateWidget(FHierarchyTableEntryData* EntryData) override;
+};
\ No newline at end of file
diff --git a/Engine/Plugins/Experimental/Animation/HierarchyTableBuiltin/Source/Editor/Private/TimeProfile/TimeProfileFactory.cpp b/Engine/Plugins/Experimental/Animation/HierarchyTableBuiltin/Source/Editor/Private/TimeProfile/TimeProfileFactory.cpp
new file mode 100644
--- /dev/null
+++ b/Engine/Plugins/Experimental/Animation/HierarchyTableBuiltin/Source/Editor/Private/TimeProfile/TimeProfileFactory.cpp
@@ -0,0 +1,15 @@
+// Copyright Epic Games, Inc. All Rights Reserved.
+
+#include "TimeProfile/TimeProfileFactory.h"
+#include "TimeProfile/TimeProfileProxyColumn.h"
+#include "TimeProfile/TimeProfileWidgetConstructor.h"
+
+void UHierarchyTableTimeFactory::RegisterWidgetConstructors(ITypedElementDataStorageInterface& DataStorage, ITypedElementDataStorageUiInterface& DataStorageUi) const
+{
+	using namespace TypedElementDataStorage;
+
+	DataStorageUi.RegisterWidgetFactory<FHierarchyTableTimeWidgetConstructor_StartTime>(FName(TEXT("General.Cell")), FColumn<FHierarchyTableTimeColumn_StartTime>());
+	DataStorageUi.RegisterWidgetFactory<FHierarchyTableTimeWidgetConstructor_EndTime>(FName(TEXT("General.Cell")), FColumn<FHierarchyTableTimeColumn_EndTime>());
+	DataStorageUi.RegisterWidgetFactory<FHierarchyTableTimeWidgetConstructor_TimeFactor>(FName(TEXT("General.Cell")), FColumn<FHierarchyTableTimeColumn_TimeFactor>());
+	DataStorageUi.RegisterWidgetFactory<FHierarchyTableTimeWidgetConstructor_Preview>(FName(TEXT("General.Cell")), FColumn<FHierarchyTableTimeColumn_Preview>());
+}
\ No newline at end of file
diff --git a/Engine/Plugins/Experimental/Animation/HierarchyTableBuiltin/Source/Editor/Private/TimeProfile/TimeProfileFactory.h b/Engine/Plugins/Experimental/Animation/HierarchyTableBuiltin/Source/Editor/Private/TimeProfile/TimeProfileFactory.h
new file mode 100644
--- /dev/null
+++ b/Engine/Plugins/Experimental/Animation/HierarchyTableBuiltin/Source/Editor/Private/TimeProfile/TimeProfileFactory.h
@@ -0,0 +1,15 @@
+// Copyright Epic Games, Inc. All Rights Reserved.
+
+#pragma once
+
+#include "Elements/Interfaces/TypedElementDataStorageFactory.h"
+#include "TimeProfileFactory.generated.h"
+
+UCLASS()
+class UHierarchyTableTimeFactory : public UTypedElementDataStorageFactory
+{
+	GENERATED_BODY()
+
+public:
+	void RegisterWidgetConstructors(ITypedElementDataStorageInterface& DataStorage, ITypedElementDataStorageUiInterface& DataStorageUi) const override;
+};
\ No newline at end of file
diff --git a/Engine/Plugins/Experimental/Animation/HierarchyTableBuiltin/Source/Editor/Private/TimeProfile/TimeProfileTypeHandler.cpp b/Engine/Plugins/Experimental/Animation/HierarchyTableBuiltin/Source/Editor/Private/TimeProfile/TimeProfileTypeHandler.cpp
new file mode 100644
--- /dev/null
+++ b/Engine/Plugins/Experimental/Animation/HierarchyTableBuiltin/Source/Editor/Private/TimeProfile/TimeProfileTypeHandler.cpp
@@ -0,0 +1,24 @@
+// Copyright Epic Games, Inc. All Rights Reserved.
+
+#pragma once
+
+#include "TimeProfile/TimeProfileTypeHandler.h"
+#include "TimeProfile/HierarchyTableTypeTime.h"
+#include "TimeProfile/TimeProfileProxyColumn.h"
+
+FInstancedStruct UHierarchyTableTypeHandler_Time::GetDefaultEntry() const
+{
+	FInstancedStruct NewEntry;
+	NewEntry.InitializeAs(FHierarchyTableType_Time::StaticStruct());
+	FHierarchyTableType_Time& NewEntryRef = NewEntry.GetMutable<FHierarchyTableType_Time>();
+	NewEntryRef.StartTime = 0.0f;
+	NewEntryRef.EndTime = 1.0f;
+	NewEntryRef.TimeFactor = 1.0f;
+
+	return NewEntry;
+}
+
+TArray<UScriptStruct*> UHierarchyTableTypeHandler_Time::GetColumns() const
+{
+	return { FHierarchyTableTimeColumn_StartTime::StaticStruct(), FHierarchyTableTimeColumn_EndTime::StaticStruct(), FHierarchyTableTimeColumn_TimeFactor::StaticStruct() };
+}
\ No newline at end of file
diff --git a/Engine/Plugins/Experimental/Animation/HierarchyTableBuiltin/Source/Editor/Private/TimeProfile/TimeProfileTypeHandler.h b/Engine/Plugins/Experimental/Animation/HierarchyTableBuiltin/Source/Editor/Private/TimeProfile/TimeProfileTypeHandler.h
new file mode 100644
--- /dev/null
+++ b/Engine/Plugins/Experimental/Animation/HierarchyTableBuiltin/Source/Editor/Private/TimeProfile/TimeProfileTypeHandler.h
@@ -0,0 +1,18 @@
+// Copyright Epic Games, Inc. All Rights Reserved.
+
+#pragma once
+
+#include "HierarchyTableTypeRegistry.h"
+
+#include "TimeProfileTypeHandler.generated.h"
+
+UCLASS()
+class UHierarchyTableTypeHandler_Time final : public UHierarchyTableTypeHandler_Base
+{
+	GENERATED_BODY()
+
+public:
+	FInstancedStruct GetDefaultEntry() const override;
+
+	TArray<UScriptStruct*> GetColumns() const override;
+};
\ No newline at end of file
diff --git a/Engine/Plugins/Experimental/Animation/HierarchyTableBuiltin/Source/Editor/Private/TimeProfile/TimeProfileWidgetConstructor.cpp b/Engine/Plugins/Experimental/Animation/HierarchyTableBuiltin/Source/Editor/Private/TimeProfile/TimeProfileWidgetConstructor.cpp
new file mode 100644
--- /dev/null
+++ b/Engine/Plugins/Experimental/Animation/HierarchyTableBuiltin/Source/Editor/Private/TimeProfile/TimeProfileWidgetConstructor.cpp
@@ -0,0 +1,63 @@
+// Copyright Epic Games, Inc. All Rights Reserved.
+
+#include "TimeProfileWidgetConstructor.h"
+#include "Widgets/Input/SSpinBox.h"
+#include "TimeProfile/HierarchyTableTypeTime.h"
+#include "HierarchyTable.h"
+
+TSharedRef<SWidget> FHierarchyTableTimeWidgetConstructor_StartTime::CreateWidget(FHierarchyTableEntryData* EntryData)
+{
+	return SNew(SSpinBox<float>)
+		.IsEnabled_Lambda([EntryData]() { return EntryData->IsOverridden(); })
+		.MinDesiredWidth(100.0f)
+		.MinValue(0.0f)
+		.MaxValue(1.0f)
+		.Value_Lambda([EntryData]()
+			{
+				return EntryData->GetValue<FHierarchyTableType_Time>()->StartTime;
+			})
+		.OnValueChanged_Lambda([EntryData](float NewValue)
+			{
+				EntryData->GetMutableValue<FHierarchyTableType_Time>()->StartTime = NewValue;
+			});
+}
+
+TSharedRef<SWidget> FHierarchyTableTimeWidgetConstructor_EndTime::CreateWidget(FHierarchyTableEntryData* EntryData)
+{
+	return SNew(SSpinBox<float>)
+		.IsEnabled_Lambda([EntryData]() { return EntryData->IsOverridden(); })
+		.MinDesiredWidth(100.0f)
+		.MinValue(0.0f)
+		.MaxValue(1.0f)
+		.Value_Lambda([EntryData]()
+			{
+				return EntryData->GetValue<FHierarchyTableType_Time>()->EndTime;
+			})
+		.OnValueChanged_Lambda([EntryData](float NewValue)
+			{
+				EntryData->GetMutableValue<FHierarchyTableType_Time>()->EndTime = NewValue;
+			});
+}
+
+TSharedRef<SWidget> FHierarchyTableTimeWidgetConstructor_TimeFactor::CreateWidget(FHierarchyTableEntryData* EntryData)
+{
+	return SNew(SSpinBox<float>)
+		.IsEnabled_Lambda([EntryData]() { return EntryData->IsOverridden(); })
+		.MinDesiredWidth(100.0f)
+		.MinValue(0.0f)
+		.MaxValue(1.0f)
+		.Value_Lambda([EntryData]()
+			{
+				return EntryData->GetValue<FHierarchyTableType_Time>()->TimeFactor;
+			})
+		.OnValueChanged_Lambda([EntryData](float NewValue)
+			{
+				EntryData->GetMutableValue<FHierarchyTableType_Time>()->TimeFactor = NewValue;
+			});
+}
+
+TSharedRef<SWidget> FHierarchyTableTimeWidgetConstructor_Preview::CreateWidget(FHierarchyTableEntryData* EntryData)
+{
+	return SNew(STextBlock)
+		.Text(INVTEXT("PREVIEW"));
+}
\ No newline at end of file
diff --git a/Engine/Plugins/Experimental/Animation/HierarchyTableBuiltin/Source/Editor/Private/TimeProfile/TimeProfileWidgetConstructor.h b/Engine/Plugins/Experimental/Animation/HierarchyTableBuiltin/Source/Editor/Private/TimeProfile/TimeProfileWidgetConstructor.h
new file mode 100644
--- /dev/null
+++ b/Engine/Plugins/Experimental/Animation/HierarchyTableBuiltin/Source/Editor/Private/TimeProfile/TimeProfileWidgetConstructor.h
@@ -0,0 +1,60 @@
+// Copyright Epic Games, Inc. All Rights Reserved.
+
+#pragma once
+
+#include "Elements/Interfaces/TypedElementDataStorageUiInterface.h"
+#include "HierarchyTableWidgetConstructor.h"
+
+#include "TimeProfileWidgetConstructor.generated.h"
+
+USTRUCT()
+struct FHierarchyTableTimeWidgetConstructor_StartTime : public FHierarchyTableWidgetConstructor
+{
+	GENERATED_BODY()
+
+public:
+	FHierarchyTableTimeWidgetConstructor_StartTime() : Super(StaticStruct()) {}
+	~FHierarchyTableTimeWidgetConstructor_StartTime() override = default;
+
+protected:
+	virtual TSharedRef<SWidget> CreateWidget(FHierarchyTableEntryData* EntryData) override;
+};
+
+USTRUCT()
+struct FHierarchyTableTimeWidgetConstructor_EndTime : public FHierarchyTableWidgetConstructor
+{
+	GENERATED_BODY()
+
+public:
+	FHierarchyTableTimeWidgetConstructor_EndTime() : Super(StaticStruct()) {}
+	~FHierarchyTableTimeWidgetConstructor_EndTime() override = default;
+
+protected:
+	virtual TSharedRef<SWidget> CreateWidget(FHierarchyTableEntryData* EntryData) override;
+};
+
+USTRUCT()
+struct FHierarchyTableTimeWidgetConstructor_TimeFactor : public FHierarchyTableWidgetConstructor
+{
+	GENERATED_BODY()
+
+public:
+	FHierarchyTableTimeWidgetConstructor_TimeFactor() : Super(StaticStruct()) {}
+	~FHierarchyTableTimeWidgetConstructor_TimeFactor() override = default;
+
+protected:
+	virtual TSharedRef<SWidget> CreateWidget(FHierarchyTableEntryData* EntryData) override;
+};
+
+USTRUCT()
+struct FHierarchyTableTimeWidgetConstructor_Preview : public FHierarchyTableWidgetConstructor
+{
+	GENERATED_BODY()
+
+public:
+	FHierarchyTableTimeWidgetConstructor_Preview() : Super(StaticStruct()) {}
+	~FHierarchyTableTimeWidgetConstructor_Preview() override = default;
+
+protected:
+	virtual TSharedRef<SWidget> CreateWidget(FHierarchyTableEntryData* EntryData) override;
+};
\ No newline at end of file
diff --git a/Engine/Plugins/Experimental/Animation/HierarchyTableBuiltin/Source/Editor/Public/HierarchyTableBuiltinEditorModule.h b/Engine/Plugins/Experimental/Animation/HierarchyTableBuiltin/Source/Editor/Public/HierarchyTableBuiltinEditorModule.h
new file mode 100644
--- /dev/null
+++ b/Engine/Plugins/Experimental/Animation/HierarchyTableBuiltin/Source/Editor/Public/HierarchyTableBuiltinEditorModule.h
@@ -0,0 +1,13 @@
+// Copyright Epic Games, Inc. All Rights Reserved.
+
+#pragma once
+
+#include "Modules/ModuleInterface.h"
+
+class FHierarchyTableBuiltinEditorModule : public IModuleInterface
+{
+public:
+	virtual void StartupModule() override;
+	
+	virtual void ShutdownModule() override;
+};
\ No newline at end of file
diff --git a/Engine/Plugins/Experimental/Animation/HierarchyTableBuiltin/Source/Editor/Public/MaskProfile/MaskProfileProxyColumn.h b/Engine/Plugins/Experimental/Animation/HierarchyTableBuiltin/Source/Editor/Public/MaskProfile/MaskProfileProxyColumn.h
new file mode 100644
--- /dev/null
+++ b/Engine/Plugins/Experimental/Animation/HierarchyTableBuiltin/Source/Editor/Public/MaskProfile/MaskProfileProxyColumn.h
@@ -0,0 +1,13 @@
+// Copyright Epic Games, Inc. All Rights Reserved.
+
+#pragma once
+
+#include "Elements/Common/TypedElementCommonTypes.h"
+
+#include "MaskProfileProxyColumn.generated.h"
+
+USTRUCT()
+struct FHierarchyTableMaskColumn_Value final : public FTypedElementDataStorageColumn
+{
+	GENERATED_BODY()
+};
\ No newline at end of file
diff --git a/Engine/Plugins/Experimental/Animation/HierarchyTableBuiltin/Source/Editor/Public/TimeProfile/TimeProfileProxyColumn.h b/Engine/Plugins/Experimental/Animation/HierarchyTableBuiltin/Source/Editor/Public/TimeProfile/TimeProfileProxyColumn.h
new file mode 100644
--- /dev/null
+++ b/Engine/Plugins/Experimental/Animation/HierarchyTableBuiltin/Source/Editor/Public/TimeProfile/TimeProfileProxyColumn.h
@@ -0,0 +1,31 @@
+// Copyright Epic Games, Inc. All Rights Reserved.
+
+#pragma once
+
+#include "Elements/Common/TypedElementCommonTypes.h"
+
+#include "TimeProfileProxyColumn.generated.h"
+
+USTRUCT()
+struct FHierarchyTableTimeColumn_StartTime final : public FTypedElementDataStorageColumn
+{
+	GENERATED_BODY()
+};
+
+USTRUCT()
+struct FHierarchyTableTimeColumn_EndTime final : public FTypedElementDataStorageColumn
+{
+	GENERATED_BODY()
+};
+
+USTRUCT()
+struct FHierarchyTableTimeColumn_TimeFactor final : public FTypedElementDataStorageColumn
+{
+	GENERATED_BODY()
+};
+
+USTRUCT()
+struct FHierarchyTableTimeColumn_Preview final : public FTypedElementDataStorageColumn
+{
+	GENERATED_BODY()
+};
\ No newline at end of file
diff --git a/Engine/Plugins/Experimental/Animation/HierarchyTableBuiltin/Source/Runtime/HierarchyTableBuiltinRuntime.Build.cs b/Engine/Plugins/Experimental/Animation/HierarchyTableBuiltin/Source/Runtime/HierarchyTableBuiltinRuntime.Build.cs
new file mode 100644
--- /dev/null
+++ b/Engine/Plugins/Experimental/Animation/HierarchyTableBuiltin/Source/Runtime/HierarchyTableBuiltinRuntime.Build.cs
@@ -0,0 +1,25 @@
+// Copyright Epic Games, Inc. All Rights Reserved.
+
+using UnrealBuildTool;
+
+public class HierarchyTableBuiltinRuntime : ModuleRules
+{
+	public HierarchyTableBuiltinRuntime(ReadOnlyTargetRules Target) : base(Target)
+	{
+		PublicDependencyModuleNames.AddRange(
+			new string[]
+			{
+				"Core",
+				"CoreUObject",
+				"Engine",
+				"HierarchyTableRuntime"
+			}
+		);
+
+		PrivateDependencyModuleNames.AddRange(
+			new string[]
+            {
+			}
+		);
+	}
+}
diff --git a/Engine/Plugins/Experimental/Animation/HierarchyTableBuiltin/Source/Runtime/Private/HierachyTableBuiltinRuntimeModule.cpp b/Engine/Plugins/Experimental/Animation/HierarchyTableBuiltin/Source/Runtime/Private/HierachyTableBuiltinRuntimeModule.cpp
new file mode 100644
--- /dev/null
+++ b/Engine/Plugins/Experimental/Animation/HierarchyTableBuiltin/Source/Runtime/Private/HierachyTableBuiltinRuntimeModule.cpp
@@ -0,0 +1,14 @@
+// Copyright Epic Games, Inc. All Rights Reserved.
+
+#include "HierarchyTableBuiltinRuntimeModule.h"
+#include "Modules/ModuleManager.h"
+
+void FHierarchyTableBuiltinRuntimeModule::StartupModule()
+{
+}
+
+void FHierarchyTableBuiltinRuntimeModule::ShutdownModule()
+{
+}
+
+IMPLEMENT_MODULE(FHierarchyTableBuiltinRuntimeModule, HierarchyTableBuiltinRuntime)
diff --git a/Engine/Plugins/Experimental/Animation/HierarchyTableBuiltin/Source/Runtime/Public/HierarchyTableBuiltinRuntimeModule.h b/Engine/Plugins/Experimental/Animation/HierarchyTableBuiltin/Source/Runtime/Public/HierarchyTableBuiltinRuntimeModule.h
new file mode 100644
--- /dev/null
+++ b/Engine/Plugins/Experimental/Animation/HierarchyTableBuiltin/Source/Runtime/Public/HierarchyTableBuiltinRuntimeModule.h
@@ -0,0 +1,13 @@
+// Copyright Epic Games, Inc. All Rights Reserved.
+
+#pragma once
+
+#include "Modules/ModuleInterface.h"
+
+class FHierarchyTableBuiltinRuntimeModule : public IModuleInterface
+{
+public:
+	virtual void StartupModule() override;
+	
+	virtual void ShutdownModule() override;
+};
\ No newline at end of file
diff --git a/Engine/Plugins/Experimental/Animation/HierarchyTableBuiltin/Source/Runtime/Public/MaskProfile/HierarchyTableTypeMask.h b/Engine/Plugins/Experimental/Animation/HierarchyTableBuiltin/Source/Runtime/Public/MaskProfile/HierarchyTableTypeMask.h
new file mode 100644
--- /dev/null
+++ b/Engine/Plugins/Experimental/Animation/HierarchyTableBuiltin/Source/Runtime/Public/MaskProfile/HierarchyTableTypeMask.h
@@ -0,0 +1,16 @@
+// Copyright Epic Games, Inc. All Rights Reserved.
+
+#pragma once
+
+#include "HierarchyTableType.h"
+
+#include "HierarchyTableTypeMask.generated.h"
+
+USTRUCT()
+struct FHierarchyTableType_Mask final : public FHierarchyTableType
+{
+	GENERATED_BODY()
+
+	UPROPERTY()
+	float Value;
+};
\ No newline at end of file
diff --git a/Engine/Plugins/Experimental/Animation/HierarchyTableBuiltin/Source/Runtime/Public/TimeProfile/HierarchyTableTypeTime.h b/Engine/Plugins/Experimental/Animation/HierarchyTableBuiltin/Source/Runtime/Public/TimeProfile/HierarchyTableTypeTime.h
new file mode 100644
--- /dev/null
+++ b/Engine/Plugins/Experimental/Animation/HierarchyTableBuiltin/Source/Runtime/Public/TimeProfile/HierarchyTableTypeTime.h
@@ -0,0 +1,22 @@
+// Copyright Epic Games, Inc. All Rights Reserved.
+
+#pragma once
+
+#include "HierarchyTableType.h"
+
+#include "HierarchyTableTypeTime.generated.h"
+
+USTRUCT()
+struct FHierarchyTableType_Time final : public FHierarchyTableType
+{
+	GENERATED_BODY()
+
+	UPROPERTY()
+	float StartTime;
+
+	UPROPERTY()
+	float EndTime;
+
+	UPROPERTY()
+	float TimeFactor;
+};
\ No newline at end of file
