Fix crash on hierarchy table builtin plugin shutdown

diff --git a/Engine/Plugins/Experimental/Animation/HierarchyTable/Source/Editor/Private/HierachyTableEditorModule.cpp b/Engine/Plugins/Experimental/Animation/HierarchyTable/Source/Editor/Private/HierachyTableEditorModule.cpp
--- a/Engine/Plugins/Experimental/Animation/HierarchyTable/Source/Editor/Private/HierachyTableEditorModule.cpp
+++ b/Engine/Plugins/Experimental/Animation/HierarchyTable/Source/Editor/Private/HierachyTableEditorModule.cpp
@@ -11,4 +11,24 @@ void FHierarchyTableEditorModule::ShutdownModule()
 {
 }
 
+void FHierarchyTableEditorModule::RegisterTableType(const UScriptStruct* HierarchyTableType, const UHierarchyTableTypeHandler_Base* Handler)
+{
+	Handlers.Add(HierarchyTableType, Handler);
+}
+
+void FHierarchyTableEditorModule::UnregisterTableType(const UScriptStruct* HierarchyTableType)
+{
+	Handlers.Remove(HierarchyTableType);
+}
+
+const UHierarchyTableTypeHandler_Base* FHierarchyTableEditorModule::FindHandler(const UScriptStruct* HierarchyTableType) const
+{
+	const UHierarchyTableTypeHandler_Base* const* Result = Handlers.Find(HierarchyTableType);
+	if (Result)
+	{
+		return *Result;
+	}
+	return nullptr;
+}
+
 IMPLEMENT_MODULE(FHierarchyTableEditorModule, HierarchyTableEditor)
diff --git a/Engine/Plugins/Experimental/Animation/HierarchyTable/Source/Editor/Private/HierarchyTableEditorToolkit.cpp b/Engine/Plugins/Experimental/Animation/HierarchyTable/Source/Editor/Private/HierarchyTableEditorToolkit.cpp
--- a/Engine/Plugins/Experimental/Animation/HierarchyTable/Source/Editor/Private/HierarchyTableEditorToolkit.cpp
+++ b/Engine/Plugins/Experimental/Animation/HierarchyTable/Source/Editor/Private/HierarchyTableEditorToolkit.cpp
@@ -15,10 +15,10 @@
 #include "ReferenceSkeleton.h"
 #include "Animation/Skeleton.h"
 #include "HierarchyTable/Columns/OverrideColumn.h"
-#include "HierarchyTableTypeRegistry.h"
 #include "HierarchyTableType.h"
 #include "ToolMenus.h"
 #include "Framework/Application/SlateApplication.h"
+#include "HierarchyTableEditorModule.h"
 
 #define LOCTEXT_NAMESPACE "HierarchyTableEditorToolkit"
 
@@ -118,7 +118,8 @@ TSharedRef<SWidget> FHierarchyTableEditorToolkit::CreateTedsOutliner()
 		return SNullWidget::NullWidget;
 	}
 
-	const UHierarchyTableTypeHandler_Base* Handler = GetDefault<UHierarchyTableTypeRegistry>()->FindHandler(HierarchyTable->TableType);
+	FHierarchyTableEditorModule& HierarchyTableModule = FModuleManager::GetModuleChecked<FHierarchyTableEditorModule>("HierarchyTableEditor");
+	const UHierarchyTableTypeHandler_Base* Handler = HierarchyTableModule.FindHandler(HierarchyTable->TableType);
 	if (!ensureMsgf(Handler, TEXT("Could not find handler for %s, have you forgotten to register it?"), *HierarchyTable->TableType->GetName()))
 	{
 		return SNullWidget::NullWidget;
diff --git a/Engine/Plugins/Experimental/Animation/HierarchyTable/Source/Editor/Private/HierarchyTableFactory.cpp b/Engine/Plugins/Experimental/Animation/HierarchyTable/Source/Editor/Private/HierarchyTableFactory.cpp
--- a/Engine/Plugins/Experimental/Animation/HierarchyTable/Source/Editor/Private/HierarchyTableFactory.cpp
+++ b/Engine/Plugins/Experimental/Animation/HierarchyTable/Source/Editor/Private/HierarchyTableFactory.cpp
@@ -11,7 +11,7 @@
 #include "Kismet2/SClassPickerDialog.h"
 #include "StructViewerFilter.h"
 #include "StructViewerModule.h"
-#include "HierarchyTableTypeRegistry.h"
+#include "HierarchyTableEditorModule.h"
 
 UHierarchyTableFactory::UHierarchyTableFactory()
 {
@@ -26,7 +26,8 @@ UObject* UHierarchyTableFactory::FactoryCreateNew(UClass* Class, UObject* InPare
 	HierarchyTable->Skeleton = Skeleton.Get();
 	HierarchyTable->TableType = TableType;
 
-	const UHierarchyTableTypeHandler_Base* Handler = GetDefault<UHierarchyTableTypeRegistry>()->FindHandler(TableType);
+	FHierarchyTableEditorModule& HierarchyTableModule = FModuleManager::GetModuleChecked<FHierarchyTableEditorModule>("HierarchyTableEditor");
+	const UHierarchyTableTypeHandler_Base* Handler = HierarchyTableModule.FindHandler(TableType);
 	check(Handler);
 
 	FInstancedStruct DefaultEntry = Handler->GetDefaultEntry();
diff --git a/Engine/Plugins/Experimental/Animation/HierarchyTable/Source/Editor/Private/HierarchyTableTypeRegistry.cpp b/Engine/Plugins/Experimental/Animation/HierarchyTable/Source/Editor/Private/HierarchyTableTypeRegistry.cpp
deleted file mode 100644
diff --git a/Engine/Plugins/Experimental/Animation/HierarchyTable/Source/Editor/Public/HierarchyTableEditorModule.h b/Engine/Plugins/Experimental/Animation/HierarchyTable/Source/Editor/Public/HierarchyTableEditorModule.h
--- a/Engine/Plugins/Experimental/Animation/HierarchyTable/Source/Editor/Public/HierarchyTableEditorModule.h
+++ b/Engine/Plugins/Experimental/Animation/HierarchyTable/Source/Editor/Public/HierarchyTableEditorModule.h
@@ -2,7 +2,8 @@
 
 #pragma once
 
-#include "Templates/SharedPointer.h"
+#include "HierarchyTableTypeHandler.h"
+#include "HierarchyTableType.h"
 #include "Modules/ModuleInterface.h"
 
 class FHierarchyTableEditorModule : public IModuleInterface
@@ -12,6 +13,12 @@ public:
 	
 	virtual void ShutdownModule() override;
 
+	void HIERARCHYTABLEEDITOR_API RegisterTableType(const UScriptStruct* HierarchyTableType, const UHierarchyTableTypeHandler_Base* Handler);
+
+	void HIERARCHYTABLEEDITOR_API UnregisterTableType(const UScriptStruct* HierarchyTableType);
+
+	const UHierarchyTableTypeHandler_Base* FindHandler(const UScriptStruct* HierarchyTableType) const;
+
 private:
-	TSharedPtr<class FHierarchyTableAssetTypeActions> HierarchyTableAssetTypeActions;
+	TMap<const UScriptStruct*, const UHierarchyTableTypeHandler_Base*> Handlers;
 };
\ No newline at end of file
diff --git a/Engine/Plugins/Experimental/Animation/HierarchyTable/Source/Editor/Public/HierarchyTableTypeHandler.h b/Engine/Plugins/Experimental/Animation/HierarchyTable/Source/Editor/Public/HierarchyTableTypeHandler.h
new file mode 100644
--- /dev/null
+++ b/Engine/Plugins/Experimental/Animation/HierarchyTable/Source/Editor/Public/HierarchyTableTypeHandler.h
@@ -0,0 +1,23 @@
+// Copyright Epic Games, Inc. All Rights Reserved.
+
+#pragma once
+
+#include "UObject/Object.h"
+#include "StructUtils/InstancedStruct.h"
+#include "HierarchyTableTypeHandler.generated.h"
+
+class UHierarchyTable;
+class UScriptStruct;
+
+UCLASS(Abstract)
+class HIERARCHYTABLEEDITOR_API UHierarchyTableTypeHandler_Base : public UObject
+{
+	GENERATED_BODY()
+
+public:
+	virtual FInstancedStruct GetDefaultEntry() const { return FInstancedStruct(); }
+
+	virtual void InitializeTable(TObjectPtr<UHierarchyTable> InHierarchyTable) const {};
+
+	virtual TArray<UScriptStruct*> GetColumns() const { return {}; };
+};
\ No newline at end of file
diff --git a/Engine/Plugins/Experimental/Animation/HierarchyTable/Source/Editor/Public/HierarchyTableTypeRegistry.h b/Engine/Plugins/Experimental/Animation/HierarchyTable/Source/Editor/Public/HierarchyTableTypeRegistry.h
deleted file mode 100644
diff --git a/Engine/Plugins/Experimental/Animation/HierarchyTableBuiltin/Source/Editor/Private/HierachyTableBuiltinEditorModule.cpp b/Engine/Plugins/Experimental/Animation/HierarchyTableBuiltin/Source/Editor/Private/HierachyTableBuiltinEditorModule.cpp
--- a/Engine/Plugins/Experimental/Animation/HierarchyTableBuiltin/Source/Editor/Private/HierachyTableBuiltinEditorModule.cpp
+++ b/Engine/Plugins/Experimental/Animation/HierarchyTableBuiltin/Source/Editor/Private/HierachyTableBuiltinEditorModule.cpp
@@ -1,25 +1,28 @@
 // Copyright Epic Games, Inc. All Rights Reserved.
 
 #include "HierarchyTableBuiltinEditorModule.h"
-#include "HierarchyTableTypeRegistry.h"
 #include "TimeProfile/HierarchyTableTypeTime.h"
 #include "TimeProfile/TimeProfileTypeHandler.h"
 #include "MaskProfile/HierarchyTableTypeMask.h"
 #include "MaskProfile/MaskProfileTypeHandler.h"
 #include "Modules/ModuleManager.h"
+#include "HierarchyTableEditorModule.h"
 
 void FHierarchyTableBuiltinEditorModule::StartupModule()
 {
-	UHierarchyTableTypeRegistry* Registry = GetMutableDefault<UHierarchyTableTypeRegistry>();
-	Registry->Register(FHierarchyTableType_Mask::StaticStruct(), GetDefault<UHierarchyTableTypeHandler_Mask>());
-	Registry->Register(FHierarchyTableType_Time::StaticStruct(), GetDefault<UHierarchyTableTypeHandler_Time>());
+	FHierarchyTableEditorModule& HierarchyTableModule = FModuleManager::GetModuleChecked<FHierarchyTableEditorModule>("HierarchyTableEditor");
+	HierarchyTableModule.RegisterTableType(FHierarchyTableType_Mask::StaticStruct(), GetDefault<UHierarchyTableTypeHandler_Mask>());
+	HierarchyTableModule.RegisterTableType(FHierarchyTableType_Time::StaticStruct(), GetDefault<UHierarchyTableTypeHandler_Time>());
 }
 
 void FHierarchyTableBuiltinEditorModule::ShutdownModule()
 {
-	UHierarchyTableTypeRegistry* Registry = GetMutableDefault<UHierarchyTableTypeRegistry>();
-	Registry->Unregister(FHierarchyTableType_Mask::StaticStruct());
-	Registry->Unregister(FHierarchyTableType_Time::StaticStruct());
+	if (FModuleManager::Get().IsModuleLoaded("HierarchyTableEditor"))
+	{
+		FHierarchyTableEditorModule& HierarchyTableModule = FModuleManager::GetModuleChecked<FHierarchyTableEditorModule>("HierarchyTableEditor");
+		HierarchyTableModule.UnregisterTableType(FHierarchyTableType_Mask::StaticStruct());
+		HierarchyTableModule.UnregisterTableType(FHierarchyTableType_Time::StaticStruct());
+	}
 }
 
 IMPLEMENT_MODULE(FHierarchyTableBuiltinEditorModule, HierarchyTableBuiltinEditor)
diff --git a/Engine/Plugins/Experimental/Animation/HierarchyTableBuiltin/Source/Editor/Private/MaskProfile/MaskProfileTypeHandler.h b/Engine/Plugins/Experimental/Animation/HierarchyTableBuiltin/Source/Editor/Private/MaskProfile/MaskProfileTypeHandler.h
--- a/Engine/Plugins/Experimental/Animation/HierarchyTableBuiltin/Source/Editor/Private/MaskProfile/MaskProfileTypeHandler.h
+++ b/Engine/Plugins/Experimental/Animation/HierarchyTableBuiltin/Source/Editor/Private/MaskProfile/MaskProfileTypeHandler.h
@@ -2,7 +2,7 @@
 
 #pragma once
 
-#include "HierarchyTableTypeRegistry.h"
+#include "HierarchyTableTypeHandler.h"
 
 #include "MaskProfileTypeHandler.generated.h"
 
diff --git a/Engine/Plugins/Experimental/Animation/HierarchyTableBuiltin/Source/Editor/Private/TimeProfile/TimeProfileTypeHandler.h b/Engine/Plugins/Experimental/Animation/HierarchyTableBuiltin/Source/Editor/Private/TimeProfile/TimeProfileTypeHandler.h
--- a/Engine/Plugins/Experimental/Animation/HierarchyTableBuiltin/Source/Editor/Private/TimeProfile/TimeProfileTypeHandler.h
+++ b/Engine/Plugins/Experimental/Animation/HierarchyTableBuiltin/Source/Editor/Private/TimeProfile/TimeProfileTypeHandler.h
@@ -2,7 +2,7 @@
 
 #pragma once
 
-#include "HierarchyTableTypeRegistry.h"
+#include "HierarchyTableTypeHandler.h"
 
 #include "TimeProfileTypeHandler.generated.h"
 
