Update hierarchy table builtin plugin from TEDS module rename and remove accidental #pragma once from some c++ files

diff --git a/Engine/Plugins/Experimental/Animation/HierarchyTableBuiltin/HierarchyTableBuiltin.uplugin b/Engine/Plugins/Experimental/Animation/HierarchyTableBuiltin/HierarchyTableBuiltin.uplugin
--- a/Engine/Plugins/Experimental/Animation/HierarchyTableBuiltin/HierarchyTableBuiltin.uplugin
+++ b/Engine/Plugins/Experimental/Animation/HierarchyTableBuiltin/HierarchyTableBuiltin.uplugin
@@ -33,7 +33,7 @@
 			"Enabled": true
 		},
 		{
-			"Name": "TypedElementDataStorageEditorUI",
+			"Name": "EditorDataStorageFeatures",
 			"Enabled": true
 		}
 	]
diff --git a/Engine/Plugins/Experimental/Animation/HierarchyTableBuiltin/Source/Editor/Private/MaskProfile/MaskProfileTypeHandler.cpp b/Engine/Plugins/Experimental/Animation/HierarchyTableBuiltin/Source/Editor/Private/MaskProfile/MaskProfileTypeHandler.cpp
--- a/Engine/Plugins/Experimental/Animation/HierarchyTableBuiltin/Source/Editor/Private/MaskProfile/MaskProfileTypeHandler.cpp
+++ b/Engine/Plugins/Experimental/Animation/HierarchyTableBuiltin/Source/Editor/Private/MaskProfile/MaskProfileTypeHandler.cpp
@@ -1,7 +1,5 @@
 // Copyright Epic Games, Inc. All Rights Reserved.
 
-#pragma once
-
 #include "MaskProfile/MaskProfileTypeHandler.h"
 #include "MaskProfile/HierarchyTableTypeMask.h"
 #include "MaskProfile/MaskProfileProxyColumn.h"
diff --git a/Engine/Plugins/Experimental/Animation/HierarchyTableBuiltin/Source/Editor/Private/TimeProfile/TimeProfileTypeHandler.cpp b/Engine/Plugins/Experimental/Animation/HierarchyTableBuiltin/Source/Editor/Private/TimeProfile/TimeProfileTypeHandler.cpp
--- a/Engine/Plugins/Experimental/Animation/HierarchyTableBuiltin/Source/Editor/Private/TimeProfile/TimeProfileTypeHandler.cpp
+++ b/Engine/Plugins/Experimental/Animation/HierarchyTableBuiltin/Source/Editor/Private/TimeProfile/TimeProfileTypeHandler.cpp
@@ -1,7 +1,5 @@
 // Copyright Epic Games, Inc. All Rights Reserved.
 
-#pragma once
-
 #include "TimeProfile/TimeProfileTypeHandler.h"
 #include "TimeProfile/HierarchyTableTypeTime.h"
 #include "TimeProfile/TimeProfileProxyColumn.h"
