Fix assertion hit due to hierarchy table editor not loaded on hierarchy table built in startup

diff --git a/Engine/Plugins/Experimental/Animation/HierarchyTableBuiltin/Source/Editor/Private/HierachyTableBuiltinEditorModule.cpp b/Engine/Plugins/Experimental/Animation/HierarchyTableBuiltin/Source/Editor/Private/HierachyTableBuiltinEditorModule.cpp
--- a/Engine/Plugins/Experimental/Animation/HierarchyTableBuiltin/Source/Editor/Private/HierachyTableBuiltinEditorModule.cpp
+++ b/Engine/Plugins/Experimental/Animation/HierarchyTableBuiltin/Source/Editor/Private/HierachyTableBuiltinEditorModule.cpp
@@ -10,7 +10,7 @@
 
 void FHierarchyTableBuiltinEditorModule::StartupModule()
 {
-	FHierarchyTableEditorModule& HierarchyTableModule = FModuleManager::GetModuleChecked<FHierarchyTableEditorModule>("HierarchyTableEditor");
+	FHierarchyTableEditorModule& HierarchyTableModule = FModuleManager::LoadModuleChecked<FHierarchyTableEditorModule>("HierarchyTableEditor");
 	HierarchyTableModule.RegisterTableType(FHierarchyTableType_Mask::StaticStruct(), GetDefault<UHierarchyTableTypeHandler_Mask>());
 	HierarchyTableModule.RegisterTableType(FHierarchyTableType_Time::StaticStruct(), GetDefault<UHierarchyTableTypeHandler_Time>());
 }
