HierarchyTable - New 'Profile Blend' node

diff --git a/Engine/Plugins/Experimental/Animation/HierarchyTable/Source/Editor/Private/HierarchyTable/Widgets/OverrideWidgetConstructor.cpp b/Engine/Plugins/Experimental/Animation/HierarchyTable/Source/Editor/Private/HierarchyTable/Widgets/OverrideWidgetConstructor.cpp
--- a/Engine/Plugins/Experimental/Animation/HierarchyTable/Source/Editor/Private/HierarchyTable/Widgets/OverrideWidgetConstructor.cpp
+++ b/Engine/Plugins/Experimental/Animation/HierarchyTable/Source/Editor/Private/HierarchyTable/Widgets/OverrideWidgetConstructor.cpp
@@ -52,8 +52,10 @@ bool FTypedElementWidgetConstructor_Override::FinalizeWidget(
 		return true;
 	}
 
-	const bool bHasParent = OverrideColumn->OwnerEntry->HasParent();
-	FHierarchyTableEntryData* EntryData = OverrideColumn->OwnerEntry; 
+	UHierarchyTable* HierarchyTable = OverrideColumn->OwnerTable;
+	const int32 EntryIndex = OverrideColumn->OwnerEntryIndex;
+
+	const bool bHasParent = HierarchyTable->TableData[EntryIndex].HasParent();
 
 	TSharedRef<SWidget> NewWidget = SNew(SHorizontalBox)
 		+ SHorizontalBox::Slot()
@@ -62,19 +64,20 @@ bool FTypedElementWidgetConstructor_Override::FinalizeWidget(
 			SNew(SButton)
 			.ButtonStyle(FAppStyle::Get(), "SimpleButton")
 			.IsEnabled(bHasParent)
-			.OnClicked_Lambda([EntryData]()
+			.OnClicked_Lambda([HierarchyTable, EntryIndex]()
 				{
-					EntryData->ToggleOverridden();
+					HierarchyTable->TableData[EntryIndex].ToggleOverridden();
 					return FReply::Handled();
 				})
 			.ContentPadding(0.0f)
 				[
 					SNew(SImage)
-					.Image_Lambda([EntryData]()
+					.Image_Lambda([HierarchyTable, EntryIndex]()
 						{
-							const bool bHasOverriddenChildren = EntryData->HasOverriddenChildren();
+							const FHierarchyTableEntryData& EntryData = HierarchyTable->TableData[EntryIndex];
+							const bool bHasOverriddenChildren = EntryData.HasOverriddenChildren();
 
-							if (EntryData->IsOverridden())
+							if (EntryData.IsOverridden())
 							{
 								if (bHasOverriddenChildren)
 								{
diff --git a/Engine/Plugins/Experimental/Animation/HierarchyTable/Source/Editor/Private/HierarchyTableEditorToolkit.cpp b/Engine/Plugins/Experimental/Animation/HierarchyTable/Source/Editor/Private/HierarchyTableEditorToolkit.cpp
--- a/Engine/Plugins/Experimental/Animation/HierarchyTable/Source/Editor/Private/HierarchyTableEditorToolkit.cpp
+++ b/Engine/Plugins/Experimental/Animation/HierarchyTable/Source/Editor/Private/HierarchyTableEditorToolkit.cpp
@@ -20,6 +20,8 @@
 #include "ToolMenus.h"
 #include "Framework/Application/SlateApplication.h"
 #include "HierarchyTableEditorModule.h"
+#include "Widgets/Input/STextEntryPopup.h"
+#include "PersonaModule.h"
 
 #define LOCTEXT_NAMESPACE "HierarchyTableEditorToolkit"
 
@@ -169,11 +171,11 @@ TSharedRef<SWidget> FHierarchyTableEditorToolkit::CreateTedsOutliner()
 		RowHandle Row = DSI->AddRow(Table);
 
 		FTypedElementOverrideColumn OverrideEntry;
-		OverrideEntry.OwnerEntry = Entry;
+		OverrideEntry.OwnerEntryIndex = EntryIndex;
 		OverrideEntry.OwnerTable = HierarchyTable;
 		DSI->AddColumn(Row, MoveTemp(OverrideEntry));
 
-		DSI->AddColumn<FTypedElementLabelColumn>(Row, { .Label = Entry->Identifier.ToString() });
+		DSI->AddColumn<FTypedElementLabelColumn>(Row, { .Label = CreateRowLabel(Entry->Identifier, Entry->EntryType) });
 
 		RowHandle* ParentRow = EntryIndexToHandleMap.Find(Entry->Parent);
 		if (ParentRow)
@@ -205,65 +207,174 @@ void FHierarchyTableEditorToolkit::ExtendToolbar()
 	{
 		FToolMenuSection& HierarchyTableSection = ToolMenu->AddSection("HierarchyTable", LOCTEXT("HierarchyTable_ToolbarLabel", "HierarchyTable"), SectionInsertLocation);
 
-		/*
 		HierarchyTableSection.AddEntry(FToolMenuEntry::InitComboButton(
 			"AddCurve",
 			FUIAction(),
-			FNewToolMenuWidget::CreateLambda([this](const FToolMenuContext& InContext) -> TSharedRef<SWidget>
+			FNewToolMenuDelegate::CreateLambda([this](UToolMenu* InSubMenu)
 				{
-					return SNew(SAnimCurvePicker, HierarchyTable->Skeleton)
-						.OnCurvePicked_Lambda([this](const FName SelectedCurve)
-							{
-								AddCurveEntry(SelectedCurve);
-								FSlateApplication::Get().DismissAllMenus();
-							});
+					FToolUIAction Action;
+				    Action.ExecuteAction = FToolMenuExecuteAction::CreateLambda([this](const FToolMenuContext& Context)
+				    {
+						TSharedRef<STextEntryPopup> TextEntry = SNew(STextEntryPopup)
+							.Label(LOCTEXT("NewCurveEntryLabal", "Curve Name"))
+							.OnTextCommitted_Lambda([this](const FText& CommittedText, ETextCommit::Type CommitInfo)
+								{
+									AddEntry(FName(CommittedText.ToString()), EHierarchyTableEntryType::Curve);
+									FSlateApplication::Get().DismissAllMenus();
+								});
+
+						FSlateApplication& SlateApp = FSlateApplication::Get();
+						SlateApp.PushMenu(
+							SlateApp.GetInteractiveTopLevelWindows()[0],
+							FWidgetPath(),
+							TextEntry,
+							SlateApp.GetCursorPos(),
+							FPopupTransitionEffect::TypeInPopup);
+				    });
+
+				    FToolMenuEntry Entry = FToolMenuEntry::InitMenuEntry(
+					    FName("AddNewCurve"),
+					    LOCTEXT("AddNewCurve_Label", "Add New Curve"),
+					    LOCTEXT("AddNewCurve_Tooltip", "Add a new curve value"),
+					    FSlateIcon(),
+					    FToolUIActionChoice(Action),
+						EUserInterfaceActionType::Button);
+				
+				    InSubMenu->AddMenuEntry("AddNewCurve", Entry);
+
+					FPersonaModule& PersonaModule = FModuleManager::LoadModuleChecked<FPersonaModule>("Persona");
+
+					InSubMenu->AddMenuEntry("ExistingCurveMenu", FToolMenuEntry::InitWidget(
+						"ExistingCurveMenu",
+						SNew(SVerticalBox)
+						+SVerticalBox::Slot()
+						.AutoHeight()
+						[
+							PersonaModule.CreateCurvePicker(HierarchyTable->Skeleton,
+								FOnCurvePicked::CreateLambda([this](const FName& InName)
+									{
+										AddEntry(InName, EHierarchyTableEntryType::Curve);
+										FSlateApplication::Get().DismissAllMenus();
+									}))
+						],
+						FText(),
+						true,
+						false,
+						true
+					));
 				}),
 			LOCTEXT("AddCurve_Label", "Add Curve"),
 			LOCTEXT("AddCurve_ToolTip", "Add a new curve to the hierarchy"),
 			FSlateIcon(FAppStyle::Get().GetStyleSetName(), "Icons.Plus")
 		));
-		*/
+
+		HierarchyTableSection.AddEntry(FToolMenuEntry::InitComboButton(
+			"AddAttribute",
+			FUIAction(),
+			FNewToolMenuDelegate::CreateLambda([this](UToolMenu* InSubMenu)
+				{
+					FToolUIAction Action;
+				    Action.ExecuteAction = FToolMenuExecuteAction::CreateLambda([this](const FToolMenuContext& Context)
+				    {
+						TSharedRef<STextEntryPopup> TextEntry = SNew(STextEntryPopup)
+							.Label(LOCTEXT("NewAttributeEntryLabal", "Attribute Name"))
+							.OnTextCommitted_Lambda([this](const FText& CommittedText, ETextCommit::Type CommitInfo)
+								{
+									AddEntry(FName(CommittedText.ToString()), EHierarchyTableEntryType::Attribute);
+									FSlateApplication::Get().DismissAllMenus();
+								});
+
+						FSlateApplication& SlateApp = FSlateApplication::Get();
+						SlateApp.PushMenu(
+							SlateApp.GetInteractiveTopLevelWindows()[0],
+							FWidgetPath(),
+							TextEntry,
+							SlateApp.GetCursorPos(),
+							FPopupTransitionEffect::TypeInPopup);
+				    });
+
+				    FToolMenuEntry Entry = FToolMenuEntry::InitMenuEntry(
+					    FName("AddNewAttribute"),
+					    LOCTEXT("AddNewAttribute_Label", "Add New Attribute"),
+					    LOCTEXT("AddNewAttribute_Tooltip", "Add a new attribute value"),
+					    FSlateIcon(),
+					    FToolUIActionChoice(Action),
+						EUserInterfaceActionType::Button);
+				
+				    InSubMenu->AddMenuEntry("AddNewAttribute", Entry);
+				}),
+			LOCTEXT("AddAttribute_Label", "Add Attribute"),
+			LOCTEXT("AddAttribute_ToolTip", "Add a new attribute to the hierarchy"),
+			FSlateIcon(FAppStyle::Get().GetStyleSetName(), "Icons.Plus")
+		));
 	}
 }
 
-void FHierarchyTableEditorToolkit::AddCurveEntry(const FName CurveName)
+void FHierarchyTableEditorToolkit::AddEntry(const FName Identifier, const EHierarchyTableEntryType EntryType)
 {
-	// TODO: Reimplement
+	check(EntryType != EHierarchyTableEntryType::Bone);
+
+	if (HierarchyTable->HasIdentifier(Identifier))
+	{
+		// Avoid adding duplicate entries
+		return;
+	}
+
+	FHierarchyTableEntryData EntryData;
+	{
+		EntryData.Identifier = Identifier;
+		EntryData.EntryType = EntryType;
+		EntryData.OwnerTable = HierarchyTable;
+		EntryData.Parent = 0;
+		EntryData.Payload = TOptional<FInstancedStruct>();
+	}
+	const int32 EntryIndex = HierarchyTable->TableData.Add(EntryData);
 
-	/*
 	using namespace UE::Editor::DataStorage;
 	UTypedElementRegistry* Registry = UTypedElementRegistry::GetInstance();
 	IEditorDataStorageProvider* DSI = Registry->GetMutableDataStorage();
 	static TableHandle Table = DSI->FindTable(FName("Editor_HierarchyTableTable"));
+	const RowHandle Row = DSI->AddRow(Table);
 
-	RowHandle* RootBoneHandle = EntryIndexToHandleMap.Find(0);
-	if (!ensure(RootBoneHandle))
 	{
-		return;
+		FTypedElementOverrideColumn OverrideEntry;
+		OverrideEntry.OwnerEntryIndex = EntryIndex;
+		OverrideEntry.OwnerTable = HierarchyTable;
+		DSI->AddColumn(Row, MoveTemp(OverrideEntry));
 	}
 
-	RowHandle Row = DSI->AddRow(Table);
-	DSI->AddColumn<FTypedElementOverrideColumn>(Row, { .bIsOverridden = false });
-	DSI->AddColumn<FTypedElementLabelColumn>(Row, { .Label = CurveName.ToString() });
-	DSI->AddColumn<FTypedElementMetadataColumn>(Row,
-		{
-			.OwnerTable = HierarchyTable,
-			.Children = TArray<RowHandle>(),
-			.Type = EHierarchyTableMetadataType::Curve,
-			.BoneIndex = INDEX_NONE,
-			.CurveName = CurveName
-		});
-	DSI->AddColumn<FTypedElementParentColumn>(Row, { .Parent = *RootBoneHandle });
-	
-	FTypedElementMetadataColumn* MetadataColumn = DSI->GetColumn<FTypedElementMetadataColumn>(*RootBoneHandle);
-	ensure(MetadataColumn);
-	MetadataColumn->Children.Add(Row);
+	// Ideally would read the label directly from the HT entry struct instead of storing it itself
+	// but this is a built-in TEDS column type that is used the the tabel viewer widget.
+	DSI->AddColumn<FTypedElementLabelColumn>(Row, { .Label = CreateRowLabel(EntryData.Identifier, EntryData.EntryType) });
+
+	const RowHandle* ParentRow = EntryIndexToHandleMap.Find(EntryData.Parent);
+	if (ParentRow)
+	{
+		DSI->AddColumn<FTableRowParentColumn>(Row, { .Parent = *ParentRow });
+	}
+
+	FHierarchyTableEditorModule& HierarchyTableModule = FModuleManager::GetModuleChecked<FHierarchyTableEditorModule>("HierarchyTableEditor");
+	const UHierarchyTableTypeHandler_Base* Handler = HierarchyTableModule.FindHandler(HierarchyTable->TableType);
+
+	for (const UScriptStruct* Column : Handler->GetColumns())
+	{
+		DSI->AddColumn(Row, Column);
+	}
 
-	const UHierarchyTableTypeHandler_Base* Handler = GetDefault<UHierarchyTableTypeRegistry>()->FindHandler(HierarchyTable->TableType);
-	check(Handler);
+	EntryIndexToHandleMap.Add(EntryIndex, Row);
+}
 
-	Handler->AddData(DSI, Row, INDEX_NONE, CurveName, HierarchyTable);
-	*/
+FString FHierarchyTableEditorToolkit::CreateRowLabel(const FName EntryIdentifier, const EHierarchyTableEntryType EntryType)
+{
+	switch (EntryType)
+	{
+		case EHierarchyTableEntryType::Curve:
+			return FString::Format(TEXT("{0} (Curve)"), { EntryIdentifier.ToString() });
+		case EHierarchyTableEntryType::Attribute:
+			return FString::Format(TEXT("{0} (Attribute)"), { EntryIdentifier.ToString() });
+		default:
+			return EntryIdentifier.ToString();
+	}
 }
 
 #undef LOCTEXT_NAMESPACE
\ No newline at end of file
diff --git a/Engine/Plugins/Experimental/Animation/HierarchyTable/Source/Editor/Private/HierarchyTableWidgetConstructor.cpp b/Engine/Plugins/Experimental/Animation/HierarchyTable/Source/Editor/Private/HierarchyTableWidgetConstructor.cpp
--- a/Engine/Plugins/Experimental/Animation/HierarchyTable/Source/Editor/Private/HierarchyTableWidgetConstructor.cpp
+++ b/Engine/Plugins/Experimental/Animation/HierarchyTable/Source/Editor/Private/HierarchyTableWidgetConstructor.cpp
@@ -11,7 +11,7 @@ FHierarchyTableWidgetConstructor::FHierarchyTableWidgetConstructor(const UScript
 {
 }
 
-TSharedRef<SWidget> FHierarchyTableWidgetConstructor::CreateInternalWidget(FHierarchyTableEntryData* EntryData)
+TSharedRef<SWidget> FHierarchyTableWidgetConstructor::CreateInternalWidget(UHierarchyTable* HierarchyTable, int32 EntryIndex)
 {
 	return SNullWidget::NullWidget;
 }
@@ -32,7 +32,7 @@ bool FHierarchyTableWidgetConstructor::FinalizeWidget(IEditorDataStorageProvider
 	UE::Editor::DataStorage::RowHandle TargetRow = DataStorage->GetColumn<FTypedElementRowReferenceColumn>(Row)->Row;
 	FTypedElementOverrideColumn* OverrideColumn = DataStorage->GetColumn<FTypedElementOverrideColumn>(TargetRow);
 
-	TSharedRef<SWidget> ActualWidget = CreateInternalWidget(OverrideColumn->OwnerEntry);
+	TSharedRef<SWidget> ActualWidget = CreateInternalWidget(OverrideColumn->OwnerTable, OverrideColumn->OwnerEntryIndex);
 	WidgetInstance->SetContent(ActualWidget);
 
 	return true;
diff --git a/Engine/Plugins/Experimental/Animation/HierarchyTable/Source/Editor/Public/HierarchyTable/Columns/OverrideColumn.h b/Engine/Plugins/Experimental/Animation/HierarchyTable/Source/Editor/Public/HierarchyTable/Columns/OverrideColumn.h
--- a/Engine/Plugins/Experimental/Animation/HierarchyTable/Source/Editor/Public/HierarchyTable/Columns/OverrideColumn.h
+++ b/Engine/Plugins/Experimental/Animation/HierarchyTable/Source/Editor/Public/HierarchyTable/Columns/OverrideColumn.h
@@ -6,7 +6,7 @@
 
 #include "OverrideColumn.generated.h"
 
-USTRUCT()
+USTRUCT(meta = (DisplayName = "Override"))
 struct FTypedElementOverrideColumn final : public FHierarchyTableProxyColumn
 {
 	GENERATED_BODY()
diff --git a/Engine/Plugins/Experimental/Animation/HierarchyTable/Source/Editor/Public/HierarchyTableEditorToolkit.h b/Engine/Plugins/Experimental/Animation/HierarchyTable/Source/Editor/Public/HierarchyTableEditorToolkit.h
--- a/Engine/Plugins/Experimental/Animation/HierarchyTable/Source/Editor/Public/HierarchyTableEditorToolkit.h
+++ b/Engine/Plugins/Experimental/Animation/HierarchyTable/Source/Editor/Public/HierarchyTableEditorToolkit.h
@@ -24,7 +24,9 @@ public:
 	void ExtendToolbar();
 
 private:
-	void AddCurveEntry(const FName CurveName);
+	FString CreateRowLabel(const FName EntryIdentifier, const EHierarchyTableEntryType EntryType);
+
+	void AddEntry(const FName Identifier, const EHierarchyTableEntryType EntryType);
 
 	TSharedRef<SWidget> CreateTedsOutliner();
 
diff --git a/Engine/Plugins/Experimental/Animation/HierarchyTable/Source/Editor/Public/HierarchyTableProxyColumn.h b/Engine/Plugins/Experimental/Animation/HierarchyTable/Source/Editor/Public/HierarchyTableProxyColumn.h
--- a/Engine/Plugins/Experimental/Animation/HierarchyTable/Source/Editor/Public/HierarchyTableProxyColumn.h
+++ b/Engine/Plugins/Experimental/Animation/HierarchyTable/Source/Editor/Public/HierarchyTableProxyColumn.h
@@ -17,5 +17,5 @@ struct FHierarchyTableProxyColumn : public FEditorDataStorageColumn
 public:
 	UHierarchyTable* OwnerTable;
 
-	FHierarchyTableEntryData* OwnerEntry;
+	int32 OwnerEntryIndex;
 };
\ No newline at end of file
diff --git a/Engine/Plugins/Experimental/Animation/HierarchyTable/Source/Editor/Public/HierarchyTableWidgetConstructor.h b/Engine/Plugins/Experimental/Animation/HierarchyTable/Source/Editor/Public/HierarchyTableWidgetConstructor.h
--- a/Engine/Plugins/Experimental/Animation/HierarchyTable/Source/Editor/Public/HierarchyTableWidgetConstructor.h
+++ b/Engine/Plugins/Experimental/Animation/HierarchyTable/Source/Editor/Public/HierarchyTableWidgetConstructor.h
@@ -6,6 +6,7 @@
 
 #include "HierarchyTableWidgetConstructor.generated.h"
 
+class UHierarchyTable;
 struct FHierarchyTableEntryData;
 
 USTRUCT()
@@ -19,7 +20,7 @@ public:
 
 	virtual ~FHierarchyTableWidgetConstructor() override = default;
 
-	virtual TSharedRef<SWidget> CreateInternalWidget(FHierarchyTableEntryData* EntryData);
+	virtual TSharedRef<SWidget> CreateInternalWidget(UHierarchyTable* HierarchyTable, int32 EntryIndex);
 
 protected:
 	// Begin FHierarchyTableWidgetConstructor
diff --git a/Engine/Plugins/Experimental/Animation/HierarchyTable/Source/Runtime/Private/HierarchyTable.cpp b/Engine/Plugins/Experimental/Animation/HierarchyTable/Source/Runtime/Private/HierarchyTable.cpp
--- a/Engine/Plugins/Experimental/Animation/HierarchyTable/Source/Runtime/Private/HierarchyTable.cpp
+++ b/Engine/Plugins/Experimental/Animation/HierarchyTable/Source/Runtime/Private/HierarchyTable.cpp
@@ -104,4 +104,20 @@ TArray<const FHierarchyTableEntryData*> UHierarchyTable::GetChildren(const FHier
 	}
 
 	return Children;
+}
+
+bool UHierarchyTable::HasIdentifier(const FName Identifier) const
+{
+	return TableData.ContainsByPredicate([Identifier](const FHierarchyTableEntryData& Entry)
+		{
+			return Entry.Identifier == Identifier;
+		});
+}
+
+FHierarchyTableEntryData* UHierarchyTable::FindEntry(const FName EntryIdentifier, const EHierarchyTableEntryType EntryType)
+{
+	return TableData.FindByPredicate([EntryIdentifier, EntryType](const FHierarchyTableEntryData& Entry)
+		{
+			return Entry.Identifier == EntryIdentifier && Entry.EntryType == EntryType;
+		});
 }
\ No newline at end of file
diff --git a/Engine/Plugins/Experimental/Animation/HierarchyTable/Source/Runtime/Public/HierarchyTable.h b/Engine/Plugins/Experimental/Animation/HierarchyTable/Source/Runtime/Public/HierarchyTable.h
--- a/Engine/Plugins/Experimental/Animation/HierarchyTable/Source/Runtime/Public/HierarchyTable.h
+++ b/Engine/Plugins/Experimental/Animation/HierarchyTable/Source/Runtime/Public/HierarchyTable.h
@@ -14,10 +14,11 @@
 class USkeleton;
 
 UENUM()
-enum class EHierarchyTableEntryType
+enum class EHierarchyTableEntryType : uint8
 {
 	Bone,
 	Curve,
+	Attribute
 };
 
 class UHierarchyTable;
@@ -28,6 +29,12 @@ struct HIERARCHYTABLERUNTIME_API FHierarchyTableEntryData
 	GENERATED_BODY()
 
 public:
+	FHierarchyTableEntryData()
+		: EntryType(EHierarchyTableEntryType::Bone)
+		, Parent(INDEX_NONE)
+	{
+	}
+
 	friend UHierarchyTable;
 
 	UPROPERTY()
@@ -100,4 +107,8 @@ public:
 	void InitializeTable(FInstancedStruct DefaultEntry);
 
 	TArray<const FHierarchyTableEntryData*> GetChildren(const FHierarchyTableEntryData& Parent) const;
+
+	bool HasIdentifier(const FName Identifier) const;
+
+	FHierarchyTableEntryData* FindEntry(const FName EntryIdentifier, const EHierarchyTableEntryType EntryType);
 };
diff --git a/Engine/Plugins/Experimental/Animation/HierarchyTableBuiltin/HierarchyTableBuiltin.uplugin b/Engine/Plugins/Experimental/Animation/HierarchyTableBuiltin/HierarchyTableBuiltin.uplugin
--- a/Engine/Plugins/Experimental/Animation/HierarchyTableBuiltin/HierarchyTableBuiltin.uplugin
+++ b/Engine/Plugins/Experimental/Animation/HierarchyTableBuiltin/HierarchyTableBuiltin.uplugin
@@ -25,6 +25,11 @@
 			"Name": "HierarchyTableBuiltinEditor",
 			"Type": "Editor",
 			"LoadingPhase": "Default"
+		},
+		{
+			"Name": "HierarchyTableBuiltinUncookedOnly",
+			"Type": "UncookedOnly",
+			"LoadingPhase": "Default"
 		}
 	],
 	"Plugins": [
diff --git a/Engine/Plugins/Experimental/Animation/HierarchyTableBuiltin/Source/Editor/Private/MaskProfile/MaskProfileWidgetConstructor.cpp b/Engine/Plugins/Experimental/Animation/HierarchyTableBuiltin/Source/Editor/Private/MaskProfile/MaskProfileWidgetConstructor.cpp
--- a/Engine/Plugins/Experimental/Animation/HierarchyTableBuiltin/Source/Editor/Private/MaskProfile/MaskProfileWidgetConstructor.cpp
+++ b/Engine/Plugins/Experimental/Animation/HierarchyTableBuiltin/Source/Editor/Private/MaskProfile/MaskProfileWidgetConstructor.cpp
@@ -5,19 +5,19 @@
 #include "MaskProfile/HierarchyTableTypeMask.h"
 #include "HierarchyTable.h"
 
-TSharedRef<SWidget> FHierarchyTableMaskWidgetConstructor_Value::CreateInternalWidget(FHierarchyTableEntryData* EntryData)
+TSharedRef<SWidget> FHierarchyTableMaskWidgetConstructor_Value::CreateInternalWidget(UHierarchyTable* HierarchyTable, int32 EntryIndex)
 {
 	return SNew(SSpinBox<float>)
-		.IsEnabled_Lambda([EntryData]() { return EntryData->IsOverridden(); })
+		.IsEnabled_Lambda([HierarchyTable, EntryIndex]() { return HierarchyTable->TableData[EntryIndex].IsOverridden(); })
 		.MinDesiredWidth(100.0f)
 		.MinValue(0.0f)
 		.MaxValue(1.0f)
-		.Value_Lambda([EntryData]()
+		.Value_Lambda([HierarchyTable, EntryIndex]()
 			{
-				return EntryData->GetValue<FHierarchyTableType_Mask>()->Value;
+				return HierarchyTable->TableData[EntryIndex].GetValue<FHierarchyTableType_Mask>()->Value;
 			})
-		.OnValueChanged_Lambda([EntryData](float NewValue)
+		.OnValueChanged_Lambda([HierarchyTable, EntryIndex](float NewValue)
 			{
-				EntryData->GetMutableValue<FHierarchyTableType_Mask>()->Value = NewValue;
+				HierarchyTable->TableData[EntryIndex].GetMutableValue<FHierarchyTableType_Mask>()->Value = NewValue;
 			});
 }
\ No newline at end of file
diff --git a/Engine/Plugins/Experimental/Animation/HierarchyTableBuiltin/Source/Editor/Private/MaskProfile/MaskProfileWidgetConstructor.h b/Engine/Plugins/Experimental/Animation/HierarchyTableBuiltin/Source/Editor/Private/MaskProfile/MaskProfileWidgetConstructor.h
--- a/Engine/Plugins/Experimental/Animation/HierarchyTableBuiltin/Source/Editor/Private/MaskProfile/MaskProfileWidgetConstructor.h
+++ b/Engine/Plugins/Experimental/Animation/HierarchyTableBuiltin/Source/Editor/Private/MaskProfile/MaskProfileWidgetConstructor.h
@@ -17,5 +17,5 @@ public:
 	~FHierarchyTableMaskWidgetConstructor_Value() override = default;
 
 protected:
-	virtual TSharedRef<SWidget> CreateInternalWidget(FHierarchyTableEntryData* EntryData) override;
+	virtual TSharedRef<SWidget> CreateInternalWidget(UHierarchyTable* HierarchyTable, int32 EntryIndex) override;
 };
\ No newline at end of file
diff --git a/Engine/Plugins/Experimental/Animation/HierarchyTableBuiltin/Source/Editor/Private/TimeProfile/TimeProfileWidgetConstructor.cpp b/Engine/Plugins/Experimental/Animation/HierarchyTableBuiltin/Source/Editor/Private/TimeProfile/TimeProfileWidgetConstructor.cpp
--- a/Engine/Plugins/Experimental/Animation/HierarchyTableBuiltin/Source/Editor/Private/TimeProfile/TimeProfileWidgetConstructor.cpp
+++ b/Engine/Plugins/Experimental/Animation/HierarchyTableBuiltin/Source/Editor/Private/TimeProfile/TimeProfileWidgetConstructor.cpp
@@ -5,58 +5,58 @@
 #include "TimeProfile/HierarchyTableTypeTime.h"
 #include "HierarchyTable.h"
 
-TSharedRef<SWidget> FHierarchyTableTimeWidgetConstructor_StartTime::CreateInternalWidget(FHierarchyTableEntryData* EntryData)
+TSharedRef<SWidget> FHierarchyTableTimeWidgetConstructor_StartTime::CreateInternalWidget(UHierarchyTable* HierarchyTable, int32 EntryIndex)
 {
 	return SNew(SSpinBox<float>)
-		.IsEnabled_Lambda([EntryData]() { return EntryData->IsOverridden(); })
+		.IsEnabled_Lambda([HierarchyTable, EntryIndex]() { return HierarchyTable->TableData[EntryIndex].IsOverridden(); })
 		.MinDesiredWidth(100.0f)
 		.MinValue(0.0f)
 		.MaxValue(1.0f)
-		.Value_Lambda([EntryData]()
+		.Value_Lambda([HierarchyTable, EntryIndex]()
 			{
-				return EntryData->GetValue<FHierarchyTableType_Time>()->StartTime;
+				return HierarchyTable->TableData[EntryIndex].GetValue<FHierarchyTableType_Time>()->StartTime;
 			})
-		.OnValueChanged_Lambda([EntryData](float NewValue)
+		.OnValueChanged_Lambda([HierarchyTable, EntryIndex](float NewValue)
 			{
-				EntryData->GetMutableValue<FHierarchyTableType_Time>()->StartTime = NewValue;
+				HierarchyTable->TableData[EntryIndex].GetMutableValue<FHierarchyTableType_Time>()->StartTime = NewValue;
 			});
 }
 
-TSharedRef<SWidget> FHierarchyTableTimeWidgetConstructor_EndTime::CreateInternalWidget(FHierarchyTableEntryData* EntryData)
+TSharedRef<SWidget> FHierarchyTableTimeWidgetConstructor_EndTime::CreateInternalWidget(UHierarchyTable* HierarchyTable, int32 EntryIndex)
 {
 	return SNew(SSpinBox<float>)
-		.IsEnabled_Lambda([EntryData]() { return EntryData->IsOverridden(); })
+		.IsEnabled_Lambda([HierarchyTable, EntryIndex]() { return HierarchyTable->TableData[EntryIndex].IsOverridden(); })
 		.MinDesiredWidth(100.0f)
 		.MinValue(0.0f)
 		.MaxValue(1.0f)
-		.Value_Lambda([EntryData]()
+		.Value_Lambda([HierarchyTable, EntryIndex]()
 			{
-				return EntryData->GetValue<FHierarchyTableType_Time>()->EndTime;
+				return HierarchyTable->TableData[EntryIndex].GetValue<FHierarchyTableType_Time>()->EndTime;
 			})
-		.OnValueChanged_Lambda([EntryData](float NewValue)
+		.OnValueChanged_Lambda([HierarchyTable, EntryIndex](float NewValue)
 			{
-				EntryData->GetMutableValue<FHierarchyTableType_Time>()->EndTime = NewValue;
+				HierarchyTable->TableData[EntryIndex].GetMutableValue<FHierarchyTableType_Time>()->EndTime = NewValue;
 			});
 }
 
-TSharedRef<SWidget> FHierarchyTableTimeWidgetConstructor_TimeFactor::CreateInternalWidget(FHierarchyTableEntryData* EntryData)
+TSharedRef<SWidget> FHierarchyTableTimeWidgetConstructor_TimeFactor::CreateInternalWidget(UHierarchyTable* HierarchyTable, int32 EntryIndex)
 {
 	return SNew(SSpinBox<float>)
-		.IsEnabled_Lambda([EntryData]() { return EntryData->IsOverridden(); })
+		.IsEnabled_Lambda([HierarchyTable, EntryIndex]() { return HierarchyTable->TableData[EntryIndex].IsOverridden(); })
 		.MinDesiredWidth(100.0f)
 		.MinValue(0.0f)
 		.MaxValue(1.0f)
-		.Value_Lambda([EntryData]()
+		.Value_Lambda([HierarchyTable, EntryIndex]()
 			{
-				return EntryData->GetValue<FHierarchyTableType_Time>()->TimeFactor;
+				return HierarchyTable->TableData[EntryIndex].GetValue<FHierarchyTableType_Time>()->TimeFactor;
 			})
-		.OnValueChanged_Lambda([EntryData](float NewValue)
+		.OnValueChanged_Lambda([HierarchyTable, EntryIndex](float NewValue)
 			{
-				EntryData->GetMutableValue<FHierarchyTableType_Time>()->TimeFactor = NewValue;
+				HierarchyTable->TableData[EntryIndex].GetMutableValue<FHierarchyTableType_Time>()->TimeFactor = NewValue;
 			});
 }
 
-TSharedRef<SWidget> FHierarchyTableTimeWidgetConstructor_Preview::CreateInternalWidget(FHierarchyTableEntryData* EntryData)
+TSharedRef<SWidget> FHierarchyTableTimeWidgetConstructor_Preview::CreateInternalWidget(UHierarchyTable* HierarchyTable, int32 EntryIndex)
 {
 	return SNew(STextBlock)
 		.Text(INVTEXT("PREVIEW"));
diff --git a/Engine/Plugins/Experimental/Animation/HierarchyTableBuiltin/Source/Editor/Private/TimeProfile/TimeProfileWidgetConstructor.h b/Engine/Plugins/Experimental/Animation/HierarchyTableBuiltin/Source/Editor/Private/TimeProfile/TimeProfileWidgetConstructor.h
--- a/Engine/Plugins/Experimental/Animation/HierarchyTableBuiltin/Source/Editor/Private/TimeProfile/TimeProfileWidgetConstructor.h
+++ b/Engine/Plugins/Experimental/Animation/HierarchyTableBuiltin/Source/Editor/Private/TimeProfile/TimeProfileWidgetConstructor.h
@@ -17,7 +17,7 @@ public:
 	~FHierarchyTableTimeWidgetConstructor_StartTime() override = default;
 
 protected:
-	virtual TSharedRef<SWidget> CreateInternalWidget(FHierarchyTableEntryData* EntryData) override;
+	virtual TSharedRef<SWidget> CreateInternalWidget(UHierarchyTable* HierarchyTable, int32 EntryIndex) override;
 };
 
 USTRUCT()
@@ -30,7 +30,7 @@ public:
 	~FHierarchyTableTimeWidgetConstructor_EndTime() override = default;
 
 protected:
-	virtual TSharedRef<SWidget> CreateInternalWidget(FHierarchyTableEntryData* EntryData) override;
+	virtual TSharedRef<SWidget> CreateInternalWidget(UHierarchyTable* HierarchyTable, int32 EntryIndex) override;
 };
 
 USTRUCT()
@@ -43,7 +43,7 @@ public:
 	~FHierarchyTableTimeWidgetConstructor_TimeFactor() override = default;
 
 protected:
-	virtual TSharedRef<SWidget> CreateInternalWidget(FHierarchyTableEntryData* EntryData) override;
+	virtual TSharedRef<SWidget> CreateInternalWidget(UHierarchyTable* HierarchyTable, int32 EntryIndex) override;
 };
 
 USTRUCT()
@@ -56,5 +56,5 @@ public:
 	~FHierarchyTableTimeWidgetConstructor_Preview() override = default;
 
 protected:
-	virtual TSharedRef<SWidget> CreateInternalWidget(FHierarchyTableEntryData* EntryData) override;
+	virtual TSharedRef<SWidget> CreateInternalWidget(UHierarchyTable* HierarchyTable, int32 EntryIndex) override;
 };
\ No newline at end of file
diff --git a/Engine/Plugins/Experimental/Animation/HierarchyTableBuiltin/Source/Runtime/Private/AnimNode_BlendProfileLayeredBlend.cpp b/Engine/Plugins/Experimental/Animation/HierarchyTableBuiltin/Source/Runtime/Private/AnimNode_BlendProfileLayeredBlend.cpp
new file mode 100644
--- /dev/null
+++ b/Engine/Plugins/Experimental/Animation/HierarchyTableBuiltin/Source/Runtime/Private/AnimNode_BlendProfileLayeredBlend.cpp
@@ -0,0 +1,342 @@
+// Copyright Epic Games, Inc. All Rights Reserved.
+
+#include "AnimNode_BlendProfileLayeredBlend.h"
+#include "AnimationRuntime.h"
+#include "Animation/AnimInstanceProxy.h"
+#include "Animation/AnimTrace.h"
+#include "Animation/AnimCurveTypes.h"
+#include "Animation/AnimCurveUtils.h"
+#include "HierarchyTable.h"
+#include "MaskProfile/HierarchyTableTypeMask.h"
+#include "Animation/AttributeTypes.h"
+#include "Animation/IAttributeBlendOperator.h"
+#include "Hash/CityHash.h"
+
+#include UE_INLINE_GENERATED_CPP_BY_NAME(AnimNode_BlendProfileLayeredBlend)
+
+void FAnimNode_BlendProfileLayeredBlend::Initialize_AnyThread(const FAnimationInitializeContext& Context)
+{
+	DECLARE_SCOPE_HIERARCHICAL_COUNTER_ANIMNODE(Initialize_AnyThread)
+	FAnimNode_Base::Initialize_AnyThread(Context);
+
+	BasePose.Initialize(Context);
+	BlendPose.Initialize(Context);
+}
+
+
+void FAnimNode_BlendProfileLayeredBlend::RebuildPerBoneBlendWeights(const USkeleton* InSkeleton)
+{
+	if (InSkeleton)
+	{
+		CreateMaskWeights(InSkeleton);
+
+		SkeletonGuid = InSkeleton->GetGuid();
+		VirtualBoneGuid = InSkeleton->GetVirtualBoneGuid();
+	}
+}
+
+bool FAnimNode_BlendProfileLayeredBlend::ArePerBoneBlendWeightsValid(const USkeleton* InSkeleton) const
+{
+	return (InSkeleton != nullptr && InSkeleton->GetGuid() == SkeletonGuid && InSkeleton->GetVirtualBoneGuid() == VirtualBoneGuid);
+}
+
+void FAnimNode_BlendProfileLayeredBlend::UpdateCachedBoneData(const FBoneContainer& RequiredBones, const USkeleton* Skeleton)
+{
+	if (RequiredBones.GetSerialNumber() == RequiredBonesSerialNumber)
+	{
+		return;
+	}
+	
+	// Update cached mask weights
+	{
+		CachedCurveMaskWeights.Empty();
+		//CachedAttributeMaskWeights.Empty();
+
+		for (const FHierarchyTableEntryData& Entry : BlendMask->TableData)
+		{
+			if (Entry.EntryType == EHierarchyTableEntryType::Curve)
+			{
+				const float EntryWeight = Entry.GetValue<FHierarchyTableType_Mask>()->Value;
+				CachedCurveMaskWeights.Add<UE::Anim::FCurveElement>({ Entry.Identifier, EntryWeight });
+			}
+			/*
+			else if (Entry.EntryType == EHierarchyTableEntryType::Attribute)
+			{
+				const float EntryWeight = Entry.GetValue<FHierarchyTableType_Mask>()->Value;
+				CachedAttributeMaskWeights.Add(FNamedFloat(Entry.Identifier, EntryWeight));
+			}
+			*/
+		}
+	}
+
+	if (!ArePerBoneBlendWeightsValid(Skeleton))
+	{
+		RebuildPerBoneBlendWeights(Skeleton);
+	}
+
+	// build desired bone weights
+	const TArray<FBoneIndexType>& RequiredBoneIndices = RequiredBones.GetBoneIndicesArray();
+	const int32 NumRequiredBones = RequiredBoneIndices.Num();
+	DesiredBoneBlendWeights.SetNumZeroed(NumRequiredBones);
+	for (int32 RequiredBoneIndex = 0; RequiredBoneIndex < NumRequiredBones; RequiredBoneIndex++)
+	{
+		const int32 SkeletonBoneIndex = RequiredBones.GetSkeletonIndex(FCompactPoseBoneIndex(RequiredBoneIndex));
+		if (ensure(SkeletonBoneIndex != INDEX_NONE))
+		{
+			DesiredBoneBlendWeights[RequiredBoneIndex] = PerBoneBlendWeights[SkeletonBoneIndex];
+		}
+	}
+
+	CurrentBoneBlendWeights.Reset(DesiredBoneBlendWeights.Num());
+	CurrentBoneBlendWeights.AddZeroed(DesiredBoneBlendWeights.Num());
+
+	//Reinitialize bone blend weights now that we have cleared them
+	UpdateDesiredBoneWeight();
+
+	RequiredBonesSerialNumber = RequiredBones.GetSerialNumber();
+}
+
+void FAnimNode_BlendProfileLayeredBlend::CacheBones_AnyThread(const FAnimationCacheBonesContext& Context)
+{
+	DECLARE_SCOPE_HIERARCHICAL_COUNTER_ANIMNODE(CacheBones_AnyThread)
+	
+	BasePose.CacheBones(Context);
+	BlendPose.CacheBones(Context);
+
+	UpdateCachedBoneData(Context.AnimInstanceProxy->GetRequiredBones(), Context.AnimInstanceProxy->GetSkeleton());
+}
+
+void FAnimNode_BlendProfileLayeredBlend::Update_AnyThread(const FAnimationUpdateContext& Context)
+{
+	DECLARE_SCOPE_HIERARCHICAL_COUNTER_ANIMNODE(Update_AnyThread)
+	bool RootMotionBlendPose = false;
+	float RootMotionWeight = 0.f;
+	const float RootMotionClearWeight = bBlendRootMotionBasedOnRootBone ? 0.f : 1.f;
+
+	if (IsLODEnabled(Context.AnimInstanceProxy))
+	{
+		GetEvaluateGraphExposedInputs().Execute(Context);
+
+		if (FAnimWeight::IsRelevant(BlendWeight))
+		{
+			UpdateCachedBoneData(Context.AnimInstanceProxy->GetRequiredBones(), Context.AnimInstanceProxy->GetSkeleton());
+			UpdateDesiredBoneWeight();
+
+			if (bBlendRootMotionBasedOnRootBone && !CurrentBoneBlendWeights.IsEmpty())
+			{
+				const float NewRootMotionWeight = CurrentBoneBlendWeights[0];
+				if (NewRootMotionWeight > ZERO_ANIMWEIGHT_THRESH)
+				{
+					RootMotionWeight = NewRootMotionWeight;
+					RootMotionBlendPose = true;
+				}
+			}
+
+			const float ThisPoseRootMotionWeight = RootMotionBlendPose ? RootMotionWeight : RootMotionClearWeight;
+			BlendPose.Update(Context.FractionalWeightAndRootMotion(BlendWeight, ThisPoseRootMotionWeight));
+		}
+	}
+
+	// initialize children
+	const float BaseRootMotionWeight = 1.f - RootMotionWeight;
+
+	if (BaseRootMotionWeight < ZERO_ANIMWEIGHT_THRESH)
+	{
+		BasePose.Update(Context.FractionalWeightAndRootMotion(1.f, BaseRootMotionWeight));
+	}
+	else
+	{
+		BasePose.Update(Context);
+	}
+}
+
+void FAnimNode_BlendProfileLayeredBlend::Evaluate_AnyThread(FPoseContext& Output)
+{
+	DECLARE_SCOPE_HIERARCHICAL_COUNTER_ANIMNODE(Evaluate_AnyThread)
+	ANIM_MT_SCOPE_CYCLE_COUNTER(BlendPosesInGraph, !IsInGameThread());
+
+	FPoseContext BasePoseContext(Output);
+	BasePose.Evaluate(BasePoseContext);
+
+	FCompactPose TargetPose;
+	FBlendedCurve TargetCurve;
+	UE::Anim::FStackAttributeContainer TargetAttributes;
+
+	if (FAnimWeight::IsRelevant(BlendWeight))
+	{
+		FPoseContext CurrentPoseContext(Output);
+		BlendPose.Evaluate(CurrentPoseContext);
+
+		TargetPose.MoveBonesFrom(CurrentPoseContext.Pose);
+		TargetCurve.MoveFrom(CurrentPoseContext.Curve);
+		TargetAttributes.MoveFrom(CurrentPoseContext.CustomAttributes);
+	}
+	else
+	{
+		TargetPose.ResetToRefPose(BasePoseContext.Pose.GetBoneContainer());
+		TargetCurve.InitFrom(Output.Curve);
+		TargetAttributes.CopyFrom(Output.CustomAttributes);
+	}
+
+	FAnimationPoseData OutputPoseData(Output);
+	
+	// Blend poses
+	FAnimationRuntime::BlendTwoPosesTogetherPerBone(BasePoseContext.Pose, TargetPose, CurrentBoneBlendWeights, OutputPoseData.GetPose());
+
+	// Blend curves
+	{
+		FBlendedCurve& OutCurve = OutputPoseData.GetCurve();
+		OutCurve.CopyFrom(BasePoseContext.Curve);
+
+		if (FAnimWeight::IsRelevant(BlendWeight))
+		{
+			FBlendedCurve FilteredCurves;
+
+			// Multiply per-curve blend weights by matching blend pose curves
+			UE::Anim::FNamedValueArrayUtils::Intersection(
+				TargetCurve,
+				CachedCurveMaskWeights,
+				[this, &FilteredCurves](const UE::Anim::FCurveElement& InBlendElement, const UE::Anim::FCurveElement& InMaskElement) mutable
+					{
+						FilteredCurves.Add(InBlendElement.Name, InBlendElement.Value * InMaskElement.Value);
+					});
+
+			// Override blend curve values with premultipled curves
+			TargetCurve.Combine(FilteredCurves);
+
+			// Remove curves that have been filtered by the mask, curves with no mask value defined remain, even with 0.0 value
+			UE::Anim::FNamedValueArrayUtils::RemoveByPredicate(
+				TargetCurve,
+				CachedCurveMaskWeights,
+				[](const UE::Anim::FCurveElement& InBaseElement, const UE::Anim::FCurveElement& InMaskElement)
+				{	
+					const bool bKeep = InMaskElement.Value == 0.0;
+					return bKeep;
+				}
+			);
+
+			// Combine base and filtered pre-multiplied blend curves
+			UE::Anim::FNamedValueArrayUtils::Union(
+				OutCurve,
+				TargetCurve,
+				[this](UE::Anim::FCurveElement& InOutBaseElement, const UE::Anim::FCurveElement& InBlendElement, UE::Anim::ENamedValueUnionFlags InFlags)
+					{
+						InOutBaseElement.Value = FMath::Lerp(InOutBaseElement.Value, InBlendElement.Value, BlendWeight);
+						InOutBaseElement.Flags |= InBlendElement.Flags; // Should this only apply when the weight is relevant?
+					});
+		}
+	}
+
+	// Blend attributes
+	{
+		using namespace UE::Anim;
+
+		FStackAttributeContainer& BaseAttributes = BasePoseContext.CustomAttributes;
+		FStackAttributeContainer& BlendAttributes = TargetAttributes;
+		FStackAttributeContainer& OutAttributes = OutputPoseData.GetAttributes();
+		OutAttributes.CopyFrom(BaseAttributes);
+
+		// Premultiply masked blend attributes
+		const TArray<TWeakObjectPtr<UScriptStruct>>& BlendUniqueTypes = BlendAttributes.GetUniqueTypes();
+
+		//CachedAttributeMaskWeights.ForEachElement([&BlendUniqueTypes, &BlendAttributes](const FNamedFloat& AttributeMask)
+		for (const FNamedFloat& AttributeMask : CachedAttributeMaskWeights)
+		{
+			for (int32 BlendUniqueTypeIndex = 0; BlendUniqueTypeIndex < BlendUniqueTypes.Num(); ++BlendUniqueTypeIndex)
+			{
+				const TArray<FAttributeId>& TypedAttributeIdentifiers = BlendAttributes.GetKeys(BlendUniqueTypeIndex);
+				const TArray<UE::Anim::TWrappedAttribute<FAnimStackAllocator>>& TypedAttributeValues = BlendAttributes.GetValues(BlendUniqueTypeIndex);
+
+				const TWeakObjectPtr<UScriptStruct> AttributeType = BlendUniqueTypes[BlendUniqueTypeIndex];
+				const IAttributeBlendOperator* Operator = UE::Anim::AttributeTypes::GetTypeOperator(AttributeType);
+
+				TWrappedAttribute<FAnimStackAllocator> DefaultData;
+				AttributeType->InitializeStruct(DefaultData.GetPtr<void>());
+
+				const int32 AttributeIndex = TypedAttributeIdentifiers.IndexOfByPredicate([&AttributeMask](const FAttributeId& InAttribute)
+					{
+						return InAttribute.GetName() == AttributeMask.Name;
+					});
+
+				if (AttributeIndex != INDEX_NONE)
+				{
+					const FAttributeId& AttributeId = TypedAttributeIdentifiers[AttributeIndex];
+
+					if (AttributeMask.Value == 0.0f)
+					{
+						BlendAttributes.Remove(AttributeType.Get(), AttributeId);
+					}
+					else
+					{
+						uint8* Value = BlendAttributes.Find(AttributeType.Get(), AttributeId);
+						check(Value != nullptr);
+
+						Operator->Interpolate(DefaultData.GetPtr<void>(), Value, AttributeMask.Value, Value);
+					}
+					break;
+				}
+			}
+		}
+
+		const TArray<const UE::Anim::FStackAttributeContainer> SourceAttributes { BaseAttributes, BlendAttributes };
+		const TArray<const float> SourceWeights { 1.0f - BlendWeight, BlendWeight };
+		static const TArray<const int32> WeightView { 0, 1 };
+
+		UE::Anim::Attributes::BlendAttributes(SourceAttributes, SourceWeights, WeightView, BaseAttributes);
+	}
+}
+
+void FAnimNode_BlendProfileLayeredBlend::GatherDebugData(FNodeDebugData& DebugData)
+{
+	DECLARE_SCOPE_HIERARCHICAL_COUNTER_ANIMNODE(GatherDebugData)
+
+	BasePose.GatherDebugData(DebugData.BranchFlow(1.f));
+	BlendPose.GatherDebugData(DebugData.BranchFlow(BlendWeight));
+}
+
+void FAnimNode_BlendProfileLayeredBlend::CreateMaskWeights(const USkeleton* Skeleton)
+{
+	if (!Skeleton)
+	{
+		return;
+	}
+
+	const FReferenceSkeleton& RefSkeleton = Skeleton->GetReferenceSkeleton();
+
+	const int32 NumBones = RefSkeleton.GetNum();
+	PerBoneBlendWeights.Reset(NumBones);
+	PerBoneBlendWeights.AddZeroed(NumBones);
+
+	check(BlendMask->Skeleton == Skeleton && BlendMask->TableType == FHierarchyTableType_Mask::StaticStruct());
+
+	for (int32 EntryIndex = 0; EntryIndex < BlendMask->TableData.Num(); ++EntryIndex)
+	{
+		const FHierarchyTableEntryData& EntryData = BlendMask->TableData[EntryIndex];
+		if (EntryData.EntryType != EHierarchyTableEntryType::Bone)
+		{
+			break;
+		}
+
+		const FHierarchyTableType_Mask* MaskValue = EntryData.GetValue<FHierarchyTableType_Mask>();
+		check(MaskValue);
+			
+		PerBoneBlendWeights[EntryIndex] = MaskValue->Value;
+	}
+}
+
+void FAnimNode_BlendProfileLayeredBlend::UpdateDesiredBoneWeight()
+{
+	check(CurrentBoneBlendWeights.Num() == DesiredBoneBlendWeights.Num());
+
+	CurrentBoneBlendWeights.Init(0, CurrentBoneBlendWeights.Num());
+
+	for (int32 BoneIndex = 0; BoneIndex < DesiredBoneBlendWeights.Num(); ++BoneIndex)
+	{
+		float TargetBlendWeight = BlendWeight * DesiredBoneBlendWeights[BoneIndex];
+		
+		if (FAnimWeight::IsRelevant(TargetBlendWeight))
+		{
+			CurrentBoneBlendWeights[BoneIndex] = TargetBlendWeight;
+		}
+	}
+}
diff --git a/Engine/Plugins/Experimental/Animation/HierarchyTableBuiltin/Source/Runtime/Public/AnimNode_BlendProfileLayeredBlend.h b/Engine/Plugins/Experimental/Animation/HierarchyTableBuiltin/Source/Runtime/Public/AnimNode_BlendProfileLayeredBlend.h
new file mode 100644
--- /dev/null
+++ b/Engine/Plugins/Experimental/Animation/HierarchyTableBuiltin/Source/Runtime/Public/AnimNode_BlendProfileLayeredBlend.h
@@ -0,0 +1,117 @@
+// Copyright Epic Games, Inc. All Rights Reserved.
+
+#pragma once
+
+#include "CoreMinimal.h"
+#include "UObject/ObjectMacros.h"
+#include "Animation/AnimTypes.h"
+#include "Animation/AnimNodeBase.h"
+#include "AnimNode_BlendProfileLayeredBlend.generated.h"
+
+class UHierarchyTable;
+enum class EHierarchyTableEntryType : uint8;
+
+USTRUCT(BlueprintInternalUseOnly)
+struct HIERARCHYTABLEBUILTINRUNTIME_API FAnimNode_BlendProfileLayeredBlend : public FAnimNode_Base
+{
+	GENERATED_BODY()
+
+public:
+	/** The source pose */
+	UPROPERTY(EditAnywhere, BlueprintReadWrite, Category = Links)
+	FPoseLink BasePose;
+
+	/** The target pose */
+	UPROPERTY(EditAnywhere, BlueprintReadWrite, Category = Links, meta = (BlueprintCompilerGeneratedDefaults))
+	FPoseLink BlendPose;
+
+	/**
+	 * The blend mask to use to control layering of the pose, curves, and attributes
+	 */
+	UPROPERTY(EditAnywhere, BlueprintReadWrite, Category = Config)
+	TObjectPtr<UHierarchyTable> BlendMask;
+
+protected:
+	// Per-bone weights for the skeleton. Serialized as these are only relative to the skeleton, but can potentially
+	// be regenerated at runtime if the GUIDs dont match
+	UPROPERTY()
+	TArray<float> PerBoneBlendWeights;
+
+	// transient data to handle weight and target weight
+	// this array changes based on required bones
+	TArray<float> DesiredBoneBlendWeights;
+	TArray<float> CurrentBoneBlendWeights;
+
+	UE::Anim::TNamedValueArray<FDefaultAllocator, UE::Anim::FCurveElement> CachedCurveMaskWeights;
+
+	struct FNamedFloat
+	{
+		FName Name;
+		float Value;
+
+		FNamedFloat(const FName InName, const float InValue)
+			: Name(InName)
+			, Value(InValue)
+		{
+		}
+	};
+	
+	TArray<FNamedFloat> CachedAttributeMaskWeights;
+
+	// Guids for skeleton used to determine whether the PerBoneBlendWeights need rebuilding
+	UPROPERTY()
+	FGuid SkeletonGuid;
+
+	// Guid for virtual bones used to determine whether the PerBoneBlendWeights need rebuilding
+	UPROPERTY()
+	FGuid VirtualBoneGuid;
+
+	/** The weight of target pose */
+	UPROPERTY(EditAnywhere, BlueprintReadWrite, Category = Runtime, meta = (BlueprintCompilerGeneratedDefaults, PinShownByDefault))
+	float BlendWeight;
+
+	// Serial number of the required bones container
+	uint16 RequiredBonesSerialNumber;
+
+	/** Whether to incorporate the per-bone blend weight of the root bone when lending root motion */
+	UPROPERTY(EditAnywhere, Category = Config)
+	bool bBlendRootMotionBasedOnRootBone;
+
+public:
+	FAnimNode_BlendProfileLayeredBlend()
+		: BlendWeight(1.0f)
+		, RequiredBonesSerialNumber(0)
+		, bBlendRootMotionBasedOnRootBone(true)
+	{
+	}
+
+	// FAnimNode_Base interface
+	virtual void Initialize_AnyThread(const FAnimationInitializeContext& Context) override;
+	virtual void CacheBones_AnyThread(const FAnimationCacheBonesContext& Context) override;
+	virtual void Update_AnyThread(const FAnimationUpdateContext& Context) override;
+	virtual void Evaluate_AnyThread(FPoseContext& Output) override;
+	virtual void GatherDebugData(FNodeDebugData& DebugData) override;
+	// End of FAnimNode_Base interface
+
+	// Invalidate the cached per-bone blend weights from the skeleton
+	void InvalidatePerBoneBlendWeights() { RequiredBonesSerialNumber = 0; SkeletonGuid = FGuid(); VirtualBoneGuid = FGuid(); }
+
+	// Invalidates the cached bone data so it is recalculated the next time this node is updated
+	void InvalidateCachedBoneData() { RequiredBonesSerialNumber = 0; }
+
+	void CreateMaskWeights(const USkeleton* Skeleton);
+
+private:
+	// Rebuild cache per bone blend weights from the skeleton
+	void RebuildPerBoneBlendWeights(const USkeleton* InSkeleton);
+
+	// Check whether per-bone blend weights are valid according to the skeleton (GUID check)
+	bool ArePerBoneBlendWeightsValid(const USkeleton* InSkeleton) const;
+
+	// Update cached data if required
+	void UpdateCachedBoneData(const FBoneContainer& RequiredBones, const USkeleton* Skeleton);
+
+	friend class UAnimGraphNode_BlendProfileLayeredBlend;
+
+	void UpdateDesiredBoneWeight();
+};
\ No newline at end of file
diff --git a/Engine/Plugins/Experimental/Animation/HierarchyTableBuiltin/Source/Runtime/Public/MaskProfile/HierarchyTableTypeMask.h b/Engine/Plugins/Experimental/Animation/HierarchyTableBuiltin/Source/Runtime/Public/MaskProfile/HierarchyTableTypeMask.h
--- a/Engine/Plugins/Experimental/Animation/HierarchyTableBuiltin/Source/Runtime/Public/MaskProfile/HierarchyTableTypeMask.h
+++ b/Engine/Plugins/Experimental/Animation/HierarchyTableBuiltin/Source/Runtime/Public/MaskProfile/HierarchyTableTypeMask.h
@@ -11,6 +11,11 @@ struct FHierarchyTableType_Mask final : public FHierarchyTableType
 {
 	GENERATED_BODY()
 
+	FHierarchyTableType_Mask()
+		: Value(0.0f)
+	{
+	}
+
 	UPROPERTY()
 	float Value;
 };
\ No newline at end of file
diff --git a/Engine/Plugins/Experimental/Animation/HierarchyTableBuiltin/Source/Runtime/Public/TimeProfile/HierarchyTableTypeTime.h b/Engine/Plugins/Experimental/Animation/HierarchyTableBuiltin/Source/Runtime/Public/TimeProfile/HierarchyTableTypeTime.h
--- a/Engine/Plugins/Experimental/Animation/HierarchyTableBuiltin/Source/Runtime/Public/TimeProfile/HierarchyTableTypeTime.h
+++ b/Engine/Plugins/Experimental/Animation/HierarchyTableBuiltin/Source/Runtime/Public/TimeProfile/HierarchyTableTypeTime.h
@@ -11,6 +11,13 @@ struct FHierarchyTableType_Time final : public FHierarchyTableType
 {
 	GENERATED_BODY()
 
+	FHierarchyTableType_Time()
+		: StartTime(0.0f)
+		, EndTime(0.0f)
+		, TimeFactor(0.0f)
+	{
+	}
+
 	UPROPERTY()
 	float StartTime;
 
diff --git a/Engine/Plugins/Experimental/Animation/HierarchyTableBuiltin/Source/UncookedOnly/HierarchyTableBuiltinUncookedOnly.Build.cs b/Engine/Plugins/Experimental/Animation/HierarchyTableBuiltin/Source/UncookedOnly/HierarchyTableBuiltinUncookedOnly.Build.cs
new file mode 100644
--- /dev/null
+++ b/Engine/Plugins/Experimental/Animation/HierarchyTableBuiltin/Source/UncookedOnly/HierarchyTableBuiltinUncookedOnly.Build.cs
@@ -0,0 +1,46 @@
+// Copyright Epic Games, Inc. All Rights Reserved.
+
+using EpicGames.Core;
+using UnrealBuildTool;
+
+public class HierarchyTableBuiltinUncookedOnly : ModuleRules
+{
+	public HierarchyTableBuiltinUncookedOnly(ReadOnlyTargetRules Target) : base(Target)
+	{
+		PublicDependencyModuleNames.AddRange(
+			new string[]
+			{
+				"AnimationCore",
+				"AnimGraph",
+				"AnimGraphRuntime",
+				"AnimationBlueprintLibrary",
+				"ToolMenus"
+			});
+
+		PrivateDependencyModuleNames.AddAll(
+			"Core",
+			"CoreUObject",
+			"Engine",
+			"InputCore",
+			"Slate",
+			"SlateCore",
+			"UnrealEd",
+			"HierarchyTableRuntime",
+			"HierarchyTableEditor",
+			"HierarchyTableBuiltinRuntime"
+		);
+
+		if (Target.bBuildEditor == true)
+		{
+			PrivateDependencyModuleNames.AddRange(
+				new string[]
+				{
+						"BlueprintGraph",
+						"EditorFramework",
+						"Kismet",
+						"UnrealEd",
+				}
+			);
+		}
+	}
+}
diff --git a/Engine/Plugins/Experimental/Animation/HierarchyTableBuiltin/Source/UncookedOnly/Private/AnimGraphNode_BlendProfileLayeredBlend.cpp b/Engine/Plugins/Experimental/Animation/HierarchyTableBuiltin/Source/UncookedOnly/Private/AnimGraphNode_BlendProfileLayeredBlend.cpp
new file mode 100644
--- /dev/null
+++ b/Engine/Plugins/Experimental/Animation/HierarchyTableBuiltin/Source/UncookedOnly/Private/AnimGraphNode_BlendProfileLayeredBlend.cpp
@@ -0,0 +1,84 @@
+// Copyright Epic Games, Inc. All Rights Reserved.
+
+#include "AnimGraphNode_BlendProfileLayeredBlend.h"
+#include "ToolMenus.h"
+#include "Kismet2/BlueprintEditorUtils.h"
+
+#include "AnimGraphCommands.h"
+#include "ScopedTransaction.h"
+
+#include "DetailLayoutBuilder.h"
+#include "Kismet2/CompilerResultsLog.h"
+
+#include "HierarchyTable.h"
+#include "MaskProfile/HierarchyTableTypeMask.h"
+
+#define LOCTEXT_NAMESPACE "AnimGraphNode_BlendProfileLayeredBlend"
+
+UAnimGraphNode_BlendProfileLayeredBlend::UAnimGraphNode_BlendProfileLayeredBlend(const FObjectInitializer& ObjectInitializer)
+	: Super(ObjectInitializer)
+{
+}
+
+FLinearColor UAnimGraphNode_BlendProfileLayeredBlend::GetNodeTitleColor() const
+{
+	return FLinearColor(0.75f, 0.75f, 0.75f);
+}
+
+FText UAnimGraphNode_BlendProfileLayeredBlend::GetTooltipText() const
+{
+	return LOCTEXT("AnimGraphNode_BlendProfileLayeredBlend_Tooltip", "Profile Blend");
+}
+
+FText UAnimGraphNode_BlendProfileLayeredBlend::GetNodeTitle(ENodeTitleType::Type TitleType) const
+{
+	return LOCTEXT("AnimGraphNode_BlendProfileLayeredBlend_Title", "Profile Blend");
+}
+
+FString UAnimGraphNode_BlendProfileLayeredBlend::GetNodeCategory() const
+{
+	return TEXT("Animation|Blends");
+}
+
+void UAnimGraphNode_BlendProfileLayeredBlend::PreloadRequiredAssets()
+{
+	TObjectPtr<UHierarchyTable> BlendMask = Node.BlendMask;
+	if (BlendMask)
+	{
+		PreloadObject(BlendMask);
+	}
+
+	Super::PreloadRequiredAssets();
+}
+
+void UAnimGraphNode_BlendProfileLayeredBlend::ValidateAnimNodeDuringCompilation(class USkeleton* ForSkeleton, class FCompilerResultsLog& MessageLog)
+{
+	UAnimGraphNode_Base::ValidateAnimNodeDuringCompilation(ForSkeleton, MessageLog);
+
+	bool bCompilationError = false;
+
+	const TObjectPtr<UHierarchyTable> BlendMask = Node.BlendMask;
+
+	if (BlendMask == nullptr)
+	{
+		MessageLog.Error(*LOCTEXT("InvalidMask", "@@ has a null blend profile.").ToString(), this, BlendMask);
+		bCompilationError = true;
+	}
+	else if (BlendMask->TableType != FHierarchyTableType_Mask::StaticStruct())
+	{
+		MessageLog.Error(*LOCTEXT("InvalidMask", "@@ has a null blend profile that is not a mask type..").ToString(), this, BlendMask);
+		bCompilationError = true;
+	}
+
+	if (bCompilationError)
+	{
+		return;
+	}
+
+	if (!Node.ArePerBoneBlendWeightsValid(ForSkeleton))
+	{
+		Node.RebuildPerBoneBlendWeights(ForSkeleton);
+	}
+}
+
+#undef LOCTEXT_NAMESPACE
diff --git a/Engine/Plugins/Experimental/Animation/HierarchyTableBuiltin/Source/UncookedOnly/Private/HierachyTableBuiltinUncookedOnlyModule.cpp b/Engine/Plugins/Experimental/Animation/HierarchyTableBuiltin/Source/UncookedOnly/Private/HierachyTableBuiltinUncookedOnlyModule.cpp
new file mode 100644
--- /dev/null
+++ b/Engine/Plugins/Experimental/Animation/HierarchyTableBuiltin/Source/UncookedOnly/Private/HierachyTableBuiltinUncookedOnlyModule.cpp
@@ -0,0 +1,14 @@
+// Copyright Epic Games, Inc. All Rights Reserved.
+
+#include "HierarchyTableBuiltinUncookedOnlyModule.h"
+#include "Modules/ModuleManager.h"
+
+void FHierarchyTableBuiltinUncookedOnlyModule::StartupModule()
+{
+}
+
+void FHierarchyTableBuiltinUncookedOnlyModule::ShutdownModule()
+{
+}
+
+IMPLEMENT_MODULE(FHierarchyTableBuiltinUncookedOnlyModule, HierarchyTableBuiltinUncookedOnly)
diff --git a/Engine/Plugins/Experimental/Animation/HierarchyTableBuiltin/Source/UncookedOnly/Public/AnimGraphNode_BlendProfileLayeredBlend.h b/Engine/Plugins/Experimental/Animation/HierarchyTableBuiltin/Source/UncookedOnly/Public/AnimGraphNode_BlendProfileLayeredBlend.h
new file mode 100644
--- /dev/null
+++ b/Engine/Plugins/Experimental/Animation/HierarchyTableBuiltin/Source/UncookedOnly/Public/AnimGraphNode_BlendProfileLayeredBlend.h
@@ -0,0 +1,30 @@
+// Copyright Epic Games, Inc. All Rights Reserved.
+
+#pragma once
+
+#include "CoreMinimal.h"
+#include "UObject/ObjectMacros.h"
+#include "AnimGraphNode_BlendListBase.h"
+#include "AnimNode_BlendProfileLayeredBlend.h"
+#include "AnimGraphNode_BlendProfileLayeredBlend.generated.h"
+
+UCLASS(MinimalAPI)
+class UAnimGraphNode_BlendProfileLayeredBlend : public UAnimGraphNode_BlendListBase
+{
+	GENERATED_UCLASS_BODY()
+
+	UPROPERTY(EditAnywhere, Category=Settings)
+	FAnimNode_BlendProfileLayeredBlend Node;
+
+	// UEdGraphNode interface
+	virtual FLinearColor GetNodeTitleColor() const override;
+	virtual FText GetTooltipText() const override;
+	virtual FText GetNodeTitle(ENodeTitleType::Type TitleType) const override;
+	// End of UEdGraphNode interface
+
+	// UAnimGraphNode_Base interface
+	virtual FString GetNodeCategory() const override;
+	virtual void PreloadRequiredAssets() override;
+	virtual void ValidateAnimNodeDuringCompilation(class USkeleton* ForSkeleton, class FCompilerResultsLog& MessageLog) override;
+	// End of UAnimGraphNode_Base interface
+};
diff --git a/Engine/Plugins/Experimental/Animation/HierarchyTableBuiltin/Source/UncookedOnly/Public/HierarchyTableBuiltinUncookedOnlyModule.h b/Engine/Plugins/Experimental/Animation/HierarchyTableBuiltin/Source/UncookedOnly/Public/HierarchyTableBuiltinUncookedOnlyModule.h
new file mode 100644
--- /dev/null
+++ b/Engine/Plugins/Experimental/Animation/HierarchyTableBuiltin/Source/UncookedOnly/Public/HierarchyTableBuiltinUncookedOnlyModule.h
@@ -0,0 +1,13 @@
+// Copyright Epic Games, Inc. All Rights Reserved.
+
+#pragma once
+
+#include "Modules/ModuleInterface.h"
+
+class FHierarchyTableBuiltinUncookedOnlyModule : public IModuleInterface
+{
+public:
+	virtual void StartupModule() override;
+	
+	virtual void ShutdownModule() override;
+};
\ No newline at end of file
diff --git a/Engine/Source/Runtime/Engine/Private/Animation/AnimationRuntime.cpp b/Engine/Source/Runtime/Engine/Private/Animation/AnimationRuntime.cpp
--- a/Engine/Source/Runtime/Engine/Private/Animation/AnimationRuntime.cpp
+++ b/Engine/Source/Runtime/Engine/Private/Animation/AnimationRuntime.cpp
@@ -524,6 +524,35 @@ void FAnimationRuntime::BlendTwoPosesTogether(const FAnimationPoseData& SourcePo
 	UE::Anim::Attributes::BlendAttributes({ SourcePoseOneData.GetAttributes(), SourcePoseTwoData.GetAttributes() }, { WeightOfPoseOne, WeightOfPoseTwo }, { 0, 1 }, OutAttributes);
 }
 
+void FAnimationRuntime::BlendTwoPosesTogetherPerBone(
+		const FCompactPose& SourcePose,
+		const FCompactPose& TargetPose,
+		const TArray<float>& TargetWeights,
+		/*out*/ FCompactPose& OutPose)
+{
+	for (FCompactPoseBoneIndex BoneIndex : OutPose.ForEachBoneIndex())
+	{
+		const float BlendWeight = TargetWeights[BoneIndex.GetInt()];
+		if (FAnimationRuntime::IsFullWeight(BlendWeight))
+		{
+			OutPose[BoneIndex] = TargetPose[BoneIndex];
+		}
+		// if it doesn't have weight, take source pose 1
+		else if (FAnimationRuntime::HasWeight(BlendWeight))
+		{
+			BlendTransform<ETransformBlendMode::Overwrite>(SourcePose[BoneIndex], OutPose[BoneIndex], 1.f - BlendWeight);
+			BlendTransform<ETransformBlendMode::Accumulate>(TargetPose[BoneIndex], OutPose[BoneIndex], BlendWeight);
+		}
+		else
+		{
+			OutPose[BoneIndex] = SourcePose[BoneIndex];
+		}
+	}
+
+	// Ensure that all of the resulting rotations are normalized
+	OutPose.NormalizeRotations();
+}
+
 void FAnimationRuntime::BlendTwoPosesTogetherPerBone(
 	const FCompactPose& SourcePose1,
 	const FCompactPose& SourcePose2,
@@ -550,28 +579,7 @@ void FAnimationRuntime::BlendTwoPosesTogetherPerBone(const FAnimationPoseData& S
 
 	const FCompactPose& SourcePoseOne = SourcePoseOneData.GetPose();
 	const FCompactPose& SourcePoseTwo = SourcePoseTwoData.GetPose();
-
-	for (FCompactPoseBoneIndex BoneIndex : OutPose.ForEachBoneIndex())
-	{
-		const float BlendWeight = WeightsOfSource2[BoneIndex.GetInt()];
-		if (FAnimationRuntime::IsFullWeight(BlendWeight))
-		{
-			OutPose[BoneIndex] = SourcePoseTwo[BoneIndex];
-		}
-		// if it doesn't have weight, take source pose 1
-		else if (FAnimationRuntime::HasWeight(BlendWeight))
-		{
-			BlendTransform<ETransformBlendMode::Overwrite>(SourcePoseOne[BoneIndex], OutPose[BoneIndex], 1.f - BlendWeight);
-			BlendTransform<ETransformBlendMode::Accumulate>(SourcePoseTwo[BoneIndex], OutPose[BoneIndex], BlendWeight);
-		}
-		else
-		{
-			OutPose[BoneIndex] = SourcePoseOne[BoneIndex];
-		}
-	}
-
-	// Ensure that all of the resulting rotations are normalized
-	OutPose.NormalizeRotations();
+	BlendTwoPosesTogetherPerBone(SourcePoseOne, SourcePoseTwo, WeightsOfSource2, OutPose);
 
 	// @note : This isn't perfect as curve can link to joint, and it would be the best to use that information
 	// but that is very expensive option as we have to have another indirect look up table to search. 
diff --git a/Engine/Source/Runtime/Engine/Public/AnimationRuntime.h b/Engine/Source/Runtime/Engine/Public/AnimationRuntime.h
--- a/Engine/Source/Runtime/Engine/Public/AnimationRuntime.h
+++ b/Engine/Source/Runtime/Engine/Public/AnimationRuntime.h
@@ -235,6 +235,12 @@ public:
 		const FAnimationPoseData& SourcePoseTwoData,
 		const TArray<float>& WeightsOfSource2,
 		/*out*/ FAnimationPoseData& OutAnimationPoseData);
+
+	static ENGINE_API void BlendTwoPosesTogetherPerBone(
+		const FCompactPose& SourcePoseOneData,
+		const FCompactPose& SourcePoseTwoData,
+		const TArray<float>& WeightsOfSource2,
+		/*out*/ FCompactPose& OutAnimationPoseData);
 		
 	/**
 	* Blends together a set of poses, each with a given weight.
