Fix LOCTEXT conflict

diff --git a/Engine/Plugins/Experimental/Animation/HierarchyTableBuiltin/Source/UncookedOnly/Private/AnimGraphNode_BlendProfileLayeredBlend.cpp b/Engine/Plugins/Experimental/Animation/HierarchyTableBuiltin/Source/UncookedOnly/Private/AnimGraphNode_BlendProfileLayeredBlend.cpp
--- a/Engine/Plugins/Experimental/Animation/HierarchyTableBuiltin/Source/UncookedOnly/Private/AnimGraphNode_BlendProfileLayeredBlend.cpp
+++ b/Engine/Plugins/Experimental/Animation/HierarchyTableBuiltin/Source/UncookedOnly/Private/AnimGraphNode_BlendProfileLayeredBlend.cpp
@@ -66,7 +66,7 @@ void UAnimGraphNode_BlendProfileLayeredBlend::ValidateAnimNodeDuringCompilation(
 	}
 	else if (BlendMask->TableType != FHierarchyTableType_Mask::StaticStruct())
 	{
-		MessageLog.Error(*LOCTEXT("InvalidMask", "@@ has a null blend profile that is not a mask type..").ToString(), this, BlendMask);
+		MessageLog.Error(*LOCTEXT("InvalidMaskType", "@@ has a null blend profile that is not a mask type.").ToString(), this, BlendMask);
 		bCompilationError = true;
 	}
 
