[HT] Add undo/redo to hierarchy tables

diff --git a/Engine/Plugins/Experimental/Animation/HierarchyTable/Source/Editor/Private/HierarchyTable/Widgets/OverrideWidgetConstructor.cpp b/Engine/Plugins/Experimental/Animation/HierarchyTable/Source/Editor/Private/HierarchyTable/Widgets/OverrideWidgetConstructor.cpp
--- a/Engine/Plugins/Experimental/Animation/HierarchyTable/Source/Editor/Private/HierarchyTable/Widgets/OverrideWidgetConstructor.cpp
+++ b/Engine/Plugins/Experimental/Animation/HierarchyTable/Source/Editor/Private/HierarchyTable/Widgets/OverrideWidgetConstructor.cpp
@@ -9,6 +9,7 @@
 #include "Widgets/SBoxPanel.h"
 #include "HierarchyTable.h"
 #include "Columns/UIPropertiesColumns.h"
+#include "ScopedTransaction.h"
 
 #define LOCTEXT_NAMESPACE "FTypedElementWidgetConstructor_Override"
 
@@ -69,6 +70,8 @@ bool FTypedElementWidgetConstructor_Override::FinalizeWidget(
 			.IsEnabled(bHasParent)
 			.OnClicked_Lambda([HierarchyTable, EntryIndex]()
 				{
+					const FScopedTransaction Transaction(LOCTEXT("ToggleOverride", "Toggle Override"));
+					HierarchyTable->Modify();
 					HierarchyTable->TableData[EntryIndex].ToggleOverridden();
 					return FReply::Handled();
 				})
diff --git a/Engine/Plugins/Experimental/Animation/HierarchyTableBuiltin/Source/Editor/Private/MaskProfile/MaskProfileWidgetConstructor.cpp b/Engine/Plugins/Experimental/Animation/HierarchyTableBuiltin/Source/Editor/Private/MaskProfile/MaskProfileWidgetConstructor.cpp
--- a/Engine/Plugins/Experimental/Animation/HierarchyTableBuiltin/Source/Editor/Private/MaskProfile/MaskProfileWidgetConstructor.cpp
+++ b/Engine/Plugins/Experimental/Animation/HierarchyTableBuiltin/Source/Editor/Private/MaskProfile/MaskProfileWidgetConstructor.cpp
@@ -4,6 +4,9 @@
 #include "Widgets/Input/SSpinBox.h"
 #include "MaskProfile/HierarchyTableTypeMask.h"
 #include "HierarchyTable.h"
+#include "Editor.h"
+
+#define LOCTEXT_NAMESPACE "FHierarchyTableMaskWidgetConstructor_Value"
 
 TSharedRef<SWidget> FHierarchyTableMaskWidgetConstructor_Value::CreateInternalWidget(UHierarchyTable* HierarchyTable, int32 EntryIndex)
 {
@@ -19,5 +22,16 @@ TSharedRef<SWidget> FHierarchyTableMaskWidgetConstructor_Value::CreateInternalWi
 		.OnValueChanged_Lambda([HierarchyTable, EntryIndex](float NewValue)
 			{
 				HierarchyTable->TableData[EntryIndex].GetMutableValue<FHierarchyTableType_Mask>()->Value = NewValue;
+			})
+		.OnBeginSliderMovement_Lambda([HierarchyTable]()
+			{
+				GEditor->BeginTransaction(LOCTEXT("SetMaskValue", "Set Mask Value"));
+				HierarchyTable->Modify();
+			})
+		.OnEndSliderMovement_Lambda([](float)
+			{
+				GEditor->EndTransaction();
 			});
-}
\ No newline at end of file
+}
+
+#undef LOCTEXT_NAMESPACE
\ No newline at end of file
