Update Hierarchy Table to be animation agnostic, create a new speciality table type for USkeleton

diff --git a/Engine/Plugins/Experimental/Animation/HierarchyTable/Source/Editor/HierarchyTableEditor.Build.cs b/Engine/Plugins/Experimental/Animation/HierarchyTable/Source/Editor/HierarchyTableEditor.Build.cs
--- a/Engine/Plugins/Experimental/Animation/HierarchyTable/Source/Editor/HierarchyTableEditor.Build.cs
+++ b/Engine/Plugins/Experimental/Animation/HierarchyTable/Source/Editor/HierarchyTableEditor.Build.cs
@@ -8,7 +8,6 @@ public class HierarchyTableEditor : ModuleRules
 	public HierarchyTableEditor(ReadOnlyTargetRules Target) : base(Target)
 	{
 		PrivateDependencyModuleNames.AddAll(
-			"AnimGraph",
 			"AssetDefinition",
 			"Core",
 			"CoreUObject",
@@ -22,7 +21,6 @@ public class HierarchyTableEditor : ModuleRules
 			"TypedElementFramework",
 			"SceneOutliner",
 			"ToolMenus",
-			"Persona",
 			"TedsTableViewer"
 		);
 	}
diff --git a/Engine/Plugins/Experimental/Animation/HierarchyTable/Source/Editor/Private/DefaultHierarchyTableTypeHandler.cpp b/Engine/Plugins/Experimental/Animation/HierarchyTable/Source/Editor/Private/DefaultHierarchyTableTypeHandler.cpp
new file mode 100644
--- /dev/null
+++ b/Engine/Plugins/Experimental/Animation/HierarchyTable/Source/Editor/Private/DefaultHierarchyTableTypeHandler.cpp
@@ -0,0 +1,26 @@
+// Copyright Epic Games, Inc. All Rights Reserved.
+
+#include "DefaultHierarchyTableTypeHandler.h"
+#include "DefaultHierarchyTableType.h"
+#include "HierarchyTable.h"
+#include "HierarchyTableEditorModule.h"
+#include "Modules/ModuleManager.h"
+
+#define LOCTEXT_NAMESPACE "UHierarchyTable_TableTypeHandler_Default"
+
+TSharedPtr<FHierarchyTable_TableType> UHierarchyTable_TableTypeHandler_Default::CreateTableType() const
+{
+	return MakeShared<FHierarchyTable_TableType_Default>();
+}
+
+void UHierarchyTable_TableTypeHandler_Default::InitialiseTable(const FInstancedStruct& TableType, TObjectPtr<UHierarchyTable> HierarchyTable) const
+{
+	// Do nothing, hierarchy is built manually
+}
+
+bool UHierarchyTable_TableTypeHandler_Default::FactoryConfigureProperties(FInstancedStruct& TableType) const
+{
+	return true;
+}
+
+#undef LOCTEXT_NAMESPACE
\ No newline at end of file
diff --git a/Engine/Plugins/Experimental/Animation/HierarchyTable/Source/Editor/Private/DefaultHierarchyTableTypeHandler.h b/Engine/Plugins/Experimental/Animation/HierarchyTable/Source/Editor/Private/DefaultHierarchyTableTypeHandler.h
new file mode 100644
--- /dev/null
+++ b/Engine/Plugins/Experimental/Animation/HierarchyTable/Source/Editor/Private/DefaultHierarchyTableTypeHandler.h
@@ -0,0 +1,20 @@
+// Copyright Epic Games, Inc. All Rights Reserved.
+
+#pragma once
+
+#include "HierarchyTableTypeHandler.h"
+
+#include "DefaultHierarchyTableTypeHandler.generated.h"
+
+UCLASS()
+class UHierarchyTable_TableTypeHandler_Default final : public UHierarchyTable_TableTypeHandler
+{
+	GENERATED_BODY()
+
+public:
+	virtual TSharedPtr<FHierarchyTable_TableType> CreateTableType() const;
+
+	virtual void InitialiseTable(const FInstancedStruct& TableType, TObjectPtr<UHierarchyTable> HierarchyTable) const;
+
+	virtual bool FactoryConfigureProperties(FInstancedStruct& TableType) const override;
+};
\ No newline at end of file
diff --git a/Engine/Plugins/Experimental/Animation/HierarchyTable/Source/Editor/Private/HierachyTableEditorModule.cpp b/Engine/Plugins/Experimental/Animation/HierarchyTable/Source/Editor/Private/HierachyTableEditorModule.cpp
--- a/Engine/Plugins/Experimental/Animation/HierarchyTable/Source/Editor/Private/HierachyTableEditorModule.cpp
+++ b/Engine/Plugins/Experimental/Animation/HierarchyTable/Source/Editor/Private/HierachyTableEditorModule.cpp
@@ -2,28 +2,52 @@
 
 #include "HierarchyTableEditorModule.h"
 #include "Modules/ModuleManager.h"
+#include "DefaultHierarchyTableType.h"
+#include "DefaultHierarchyTableTypeHandler.h"
 
 void FHierarchyTableEditorModule::StartupModule()
 {
+	RegisterTableType(FHierarchyTable_TableType_Default::StaticStruct(), GetDefault<UHierarchyTable_TableTypeHandler_Default>());
 }
 
 void FHierarchyTableEditorModule::ShutdownModule()
 {
+	UnregisterTableType(FHierarchyTable_TableType_Default::StaticStruct());
 }
 
-void FHierarchyTableEditorModule::RegisterTableType(const UScriptStruct* HierarchyTableType, const UHierarchyTableTypeHandler_Base* Handler)
+void FHierarchyTableEditorModule::RegisterTableType(const UScriptStruct* TableType, const UHierarchyTable_TableTypeHandler* Handler)
 {
-	Handlers.Add(HierarchyTableType, Handler);
+	TableHandlers.Add(TableType, Handler);
 }
 
-void FHierarchyTableEditorModule::UnregisterTableType(const UScriptStruct* HierarchyTableType)
+void FHierarchyTableEditorModule::UnregisterTableType(const UScriptStruct* TableType)
 {
-	Handlers.Remove(HierarchyTableType);
+	TableHandlers.Remove(TableType);
 }
 
-const UHierarchyTableTypeHandler_Base* FHierarchyTableEditorModule::FindHandler(const UScriptStruct* HierarchyTableType) const
+const UHierarchyTable_TableTypeHandler* FHierarchyTableEditorModule::FindTableHandler(const UScriptStruct* TableType) const
 {
-	const UHierarchyTableTypeHandler_Base* const* Result = Handlers.Find(HierarchyTableType);
+	const UHierarchyTable_TableTypeHandler* const* Result = TableHandlers.Find(TableType);
+	if (Result)
+	{
+		return *Result;
+	}
+	return nullptr;
+}
+
+void FHierarchyTableEditorModule::RegisterElementType(const UScriptStruct* ElementType, const UHierarchyTable_ElementTypeHandler* Handler)
+{
+	ElementHandlers.Add(ElementType, Handler);
+}
+
+void FHierarchyTableEditorModule::UnregisterElementType(const UScriptStruct* ElementType)
+{
+	ElementHandlers.Remove(ElementType);
+}
+
+const UHierarchyTable_ElementTypeHandler* FHierarchyTableEditorModule::FindElementHandler(const UScriptStruct* ElementType) const
+{
+	const UHierarchyTable_ElementTypeHandler* const* Result = ElementHandlers.Find(ElementType);
 	if (Result)
 	{
 		return *Result;
diff --git a/Engine/Plugins/Experimental/Animation/HierarchyTable/Source/Editor/Private/HierarchyTable/Widgets/OverrideWidgetConstructor.cpp b/Engine/Plugins/Experimental/Animation/HierarchyTable/Source/Editor/Private/HierarchyTable/Widgets/OverrideWidgetConstructor.cpp
--- a/Engine/Plugins/Experimental/Animation/HierarchyTable/Source/Editor/Private/HierarchyTable/Widgets/OverrideWidgetConstructor.cpp
+++ b/Engine/Plugins/Experimental/Animation/HierarchyTable/Source/Editor/Private/HierarchyTable/Widgets/OverrideWidgetConstructor.cpp
@@ -59,7 +59,7 @@ bool FTypedElementWidgetConstructor_Override::FinalizeWidget(
 	UHierarchyTable* HierarchyTable = OverrideColumn->OwnerTable;
 	const int32 EntryIndex = OverrideColumn->OwnerEntryIndex;
 
-	const bool bHasParent = HierarchyTable->TableData[EntryIndex].HasParent();
+	const bool bHasParent = HierarchyTable->GetTableEntry(EntryIndex)->HasParent();
 
 	TSharedRef<SWidget> NewWidget = SNew(SHorizontalBox)
 		+ SHorizontalBox::Slot()
@@ -72,7 +72,7 @@ bool FTypedElementWidgetConstructor_Override::FinalizeWidget(
 				{
 					const FScopedTransaction Transaction(LOCTEXT("ToggleOverride", "Toggle Override"));
 					HierarchyTable->Modify();
-					HierarchyTable->TableData[EntryIndex].ToggleOverridden();
+					HierarchyTable->GetMutableTableEntry(EntryIndex)->ToggleOverridden();
 					return FReply::Handled();
 				})
 			.ContentPadding(0.0f)
@@ -80,10 +80,10 @@ bool FTypedElementWidgetConstructor_Override::FinalizeWidget(
 					SNew(SImage)
 					.Image_Lambda([HierarchyTable, EntryIndex]()
 						{
-							const FHierarchyTableEntryData& EntryData = HierarchyTable->TableData[EntryIndex];
-							const bool bHasOverriddenChildren = EntryData.HasOverriddenChildren();
+							const FHierarchyTableEntryData* const EntryData = HierarchyTable->GetTableEntry(EntryIndex);
+							const bool bHasOverriddenChildren = EntryData->HasOverriddenChildren();
 
-							if (EntryData.IsOverridden())
+							if (EntryData->IsOverridden())
 							{
 								if (bHasOverriddenChildren)
 								{
diff --git a/Engine/Plugins/Experimental/Animation/HierarchyTable/Source/Editor/Private/HierarchyTableEditorToolkit.cpp b/Engine/Plugins/Experimental/Animation/HierarchyTable/Source/Editor/Private/HierarchyTableEditorToolkit.cpp
--- a/Engine/Plugins/Experimental/Animation/HierarchyTable/Source/Editor/Private/HierarchyTableEditorToolkit.cpp
+++ b/Engine/Plugins/Experimental/Animation/HierarchyTable/Source/Editor/Private/HierarchyTableEditorToolkit.cpp
@@ -96,6 +96,8 @@ void FHierarchyTableEditorToolkit::OnClose()
 	{
 		DSI->RemoveRow(Row.Value);
 	}
+	
+	HierarchyTable->UnregisterOnEntryAdded(OnEntryAddedHandle);
 }
 
 void FHierarchyTableEditorToolkit::UnregisterTabSpawners(const TSharedRef<class FTabManager>& InTabManager)
@@ -116,14 +118,14 @@ TSharedRef<SWidget> FHierarchyTableEditorToolkit::CreateTedsOutliner()
 			.Text(INVTEXT("You need to enable the Typed Element Data Storage plugin to see the table viewer!"));
 	}
 
-	if (!ensure(HierarchyTable->TableType))
+	if (!ensure(HierarchyTable->ElementType))
 	{
 		return SNullWidget::NullWidget;
 	}
 
 	FHierarchyTableEditorModule& HierarchyTableModule = FModuleManager::GetModuleChecked<FHierarchyTableEditorModule>("HierarchyTableEditor");
-	const UHierarchyTableTypeHandler_Base* Handler = HierarchyTableModule.FindHandler(HierarchyTable->TableType);
-	if (!ensureMsgf(Handler, TEXT("Could not find handler for %s, have you forgotten to register it?"), *HierarchyTable->TableType->GetName()))
+	const UHierarchyTable_ElementTypeHandler* Handler = HierarchyTableModule.FindElementHandler(HierarchyTable->ElementType);
+	if (!ensureMsgf(Handler, TEXT("Could not find handler for %s, have you forgotten to register it?"), *HierarchyTable->ElementType->GetName()))
 	{
 		return SNullWidget::NullWidget;
 	}
@@ -164,9 +166,9 @@ TSharedRef<SWidget> FHierarchyTableEditorToolkit::CreateTedsOutliner()
 	
 	TArray<UScriptStruct*> BaseHierarchyTableTypeColumns = Handler->GetColumns();
 
-	for (int32 EntryIndex = 0; EntryIndex < HierarchyTable->TableData.Num(); ++EntryIndex)
+	for (int32 EntryIndex = 0; EntryIndex < HierarchyTable->GetTableData().Num(); ++EntryIndex)
 	{
-		FHierarchyTableEntryData* Entry = &HierarchyTable->TableData[EntryIndex];
+		const FHierarchyTableEntryData* const Entry = HierarchyTable->GetTableEntry(EntryIndex);
 
 		RowHandle Row = DSI->AddRow(Table);
 
@@ -194,140 +196,34 @@ TSharedRef<SWidget> FHierarchyTableEditorToolkit::CreateTedsOutliner()
 	TSharedRef<ISceneOutliner> TedsOutliner = TedsOutlinerModule.CreateTedsOutliner(InitOptions, Params, InitialColumnQuery);
 	TedsOutlinerPtr = TedsOutliner;
 
+	OnEntryAddedHandle = HierarchyTable->RegisterOnEntryAdded(UHierarchyTable::FOnEntryAdded::FDelegate::CreateRaw(this, &FHierarchyTableEditorToolkit::HandleOnEntryAdded));
+
 	return TedsOutliner;
 }
 
 void FHierarchyTableEditorToolkit::ExtendToolbar()
 {
+	FHierarchyTableEditorModule& HierarchyTableModule = FModuleManager::GetModuleChecked<FHierarchyTableEditorModule>("HierarchyTableEditor");
+	const UHierarchyTable_TableTypeHandler* Handler = HierarchyTableModule.FindTableHandler(HierarchyTable->TableMetadata.GetScriptStruct());
+	if (!ensure(Handler))
+	{
+		return;
+	}
+
 	FToolMenuOwnerScoped OwnerScoped(this);
 
 	FName ParentName;
 	static const FName MenuName = GetToolMenuToolbarName(ParentName);
 
 	UToolMenu* ToolMenu = UToolMenus::Get()->ExtendMenu(MenuName);
-	const FToolMenuInsert SectionInsertLocation("Asset", EToolMenuInsertType::After);
-	{
-		FToolMenuSection& HierarchyTableSection = ToolMenu->AddSection("HierarchyTable", LOCTEXT("HierarchyTable_ToolbarLabel", "HierarchyTable"), SectionInsertLocation);
-
-		HierarchyTableSection.AddEntry(FToolMenuEntry::InitComboButton(
-			"AddCurve",
-			FUIAction(),
-			FNewToolMenuDelegate::CreateLambda([this](UToolMenu* InSubMenu)
-				{
-					FToolUIAction Action;
-				    Action.ExecuteAction = FToolMenuExecuteAction::CreateLambda([this](const FToolMenuContext& Context)
-				    {
-						TSharedRef<STextEntryPopup> TextEntry = SNew(STextEntryPopup)
-							.Label(LOCTEXT("NewCurveEntryLabal", "Curve Name"))
-							.OnTextCommitted_Lambda([this](const FText& CommittedText, ETextCommit::Type CommitInfo)
-								{
-									AddEntry(FName(CommittedText.ToString()), EHierarchyTableEntryType::Curve);
-									FSlateApplication::Get().DismissAllMenus();
-								});
-
-						FSlateApplication& SlateApp = FSlateApplication::Get();
-						SlateApp.PushMenu(
-							SlateApp.GetInteractiveTopLevelWindows()[0],
-							FWidgetPath(),
-							TextEntry,
-							SlateApp.GetCursorPos(),
-							FPopupTransitionEffect::TypeInPopup);
-				    });
-
-				    FToolMenuEntry Entry = FToolMenuEntry::InitMenuEntry(
-					    FName("AddNewCurve"),
-					    LOCTEXT("AddNewCurve_Label", "Add New Curve"),
-					    LOCTEXT("AddNewCurve_Tooltip", "Add a new curve value"),
-					    FSlateIcon(),
-					    FToolUIActionChoice(Action),
-						EUserInterfaceActionType::Button);
-				
-				    InSubMenu->AddMenuEntry("AddNewCurve", Entry);
-
-					FPersonaModule& PersonaModule = FModuleManager::LoadModuleChecked<FPersonaModule>("Persona");
-
-					InSubMenu->AddMenuEntry("ExistingCurveMenu", FToolMenuEntry::InitWidget(
-						"ExistingCurveMenu",
-						SNew(SVerticalBox)
-						+SVerticalBox::Slot()
-						.AutoHeight()
-						[
-							PersonaModule.CreateCurvePicker(HierarchyTable->Skeleton,
-								FOnCurvePicked::CreateLambda([this](const FName& InName)
-									{
-										AddEntry(InName, EHierarchyTableEntryType::Curve);
-										FSlateApplication::Get().DismissAllMenus();
-									}))
-						],
-						FText(),
-						true,
-						false,
-						true
-					));
-				}),
-			LOCTEXT("AddCurve_Label", "Add Curve"),
-			LOCTEXT("AddCurve_ToolTip", "Add a new curve to the hierarchy"),
-			FSlateIcon(FAppStyle::Get().GetStyleSetName(), "Icons.Plus")
-		));
-
-		HierarchyTableSection.AddEntry(FToolMenuEntry::InitComboButton(
-			"AddAttribute",
-			FUIAction(),
-			FNewToolMenuDelegate::CreateLambda([this](UToolMenu* InSubMenu)
-				{
-					FToolUIAction Action;
-				    Action.ExecuteAction = FToolMenuExecuteAction::CreateLambda([this](const FToolMenuContext& Context)
-				    {
-						TSharedRef<STextEntryPopup> TextEntry = SNew(STextEntryPopup)
-							.Label(LOCTEXT("NewAttributeEntryLabal", "Attribute Name"))
-							.OnTextCommitted_Lambda([this](const FText& CommittedText, ETextCommit::Type CommitInfo)
-								{
-									AddEntry(FName(CommittedText.ToString()), EHierarchyTableEntryType::Attribute);
-									FSlateApplication::Get().DismissAllMenus();
-								});
-
-						FSlateApplication& SlateApp = FSlateApplication::Get();
-						SlateApp.PushMenu(
-							SlateApp.GetInteractiveTopLevelWindows()[0],
-							FWidgetPath(),
-							TextEntry,
-							SlateApp.GetCursorPos(),
-							FPopupTransitionEffect::TypeInPopup);
-				    });
-
-				    FToolMenuEntry Entry = FToolMenuEntry::InitMenuEntry(
-					    FName("AddNewAttribute"),
-					    LOCTEXT("AddNewAttribute_Label", "Add New Attribute"),
-					    LOCTEXT("AddNewAttribute_Tooltip", "Add a new attribute value"),
-					    FSlateIcon(),
-					    FToolUIActionChoice(Action),
-						EUserInterfaceActionType::Button);
-				
-				    InSubMenu->AddMenuEntry("AddNewAttribute", Entry);
-				}),
-			LOCTEXT("AddAttribute_Label", "Add Attribute"),
-			LOCTEXT("AddAttribute_ToolTip", "Add a new attribute to the hierarchy"),
-			FSlateIcon(FAppStyle::Get().GetStyleSetName(), "Icons.Plus")
-		));
-	}
+	Handler->ExtendToolbar(ToolMenu, HierarchyTable, this);
 }
 
-void FHierarchyTableEditorToolkit::AddEntry(const FName Identifier, const EHierarchyTableEntryType EntryType)
+int32 FHierarchyTableEditorToolkit::GetSelectedEntryIndex() const
 {
-	check(EntryType != EHierarchyTableEntryType::Bone);
-
-	if (HierarchyTable->HasIdentifier(Identifier))
-	{
-		// Avoid adding duplicate entries
-		return;
-	}
-
 	using namespace UE::Editor::DataStorage;
 	IEditorDataStorageProvider* DSI = GetMutableDataStorageFeature<IEditorDataStorageProvider>(StorageFeatureName);
-	static TableHandle Table = DSI->FindTable(FName("Editor_HierarchyTableTable"));
-
-	int32 ParentIndex = 0; // root of hierarchy
-
+	
 	TArray<FSceneOutlinerTreeItemPtr> Selection = TedsOutlinerPtr->GetTree().GetSelectedItems();
 	if (Selection.Num() > 0)
 	{
@@ -335,21 +231,28 @@ void FHierarchyTableEditorToolkit::AddEntry(const FName Identifier, const EHiera
 		const UE::Editor::DataStorage::RowHandle ParentRowHandle = TedsItem->GetRowHandle();
 
 		const FTypedElementOverrideColumn* OverrideColumn = DSI->GetColumn<FTypedElementOverrideColumn>(ParentRowHandle);
-		ParentIndex = OverrideColumn->OwnerEntryIndex;
+		return OverrideColumn->OwnerEntryIndex;
 	}
 
-	FHierarchyTableEntryData EntryData;
+	return INDEX_NONE;
+}
+
+void FHierarchyTableEditorToolkit::HandleOnEntryAdded(const int32 EntryIndex)
+{
+	const FHierarchyTableEntryData* const EntryData = HierarchyTable->GetTableEntry(EntryIndex);
+
+	using namespace UE::Editor::DataStorage;
+	IEditorDataStorageProvider* DSI = GetMutableDataStorageFeature<IEditorDataStorageProvider>(StorageFeatureName);
+	static TableHandle Table = DSI->FindTable(FName("Editor_HierarchyTableTable"));
+
+	FHierarchyTableEditorModule& HierarchyTableModule = FModuleManager::GetModuleChecked<FHierarchyTableEditorModule>("HierarchyTableEditor");
+	const UHierarchyTable_ElementTypeHandler* ElementHandler = HierarchyTableModule.FindElementHandler(HierarchyTable->ElementType);
+	if (!ensure(ElementHandler))
 	{
-		EntryData.Identifier = Identifier;
-		EntryData.EntryType = EntryType;
-		EntryData.OwnerTable = HierarchyTable;
-		EntryData.Parent = ParentIndex;
-		EntryData.Payload = TOptional<FInstancedStruct>();
+		return;
 	}
-	const int32 EntryIndex = HierarchyTable->TableData.Add(EntryData);
-
+	
 	const RowHandle Row = DSI->AddRow(Table);
-
 	{
 		FTypedElementOverrideColumn OverrideEntry;
 		OverrideEntry.OwnerEntryIndex = EntryIndex;
@@ -359,18 +262,15 @@ void FHierarchyTableEditorToolkit::AddEntry(const FName Identifier, const EHiera
 
 	// Ideally would read the label directly from the HT entry struct instead of storing it itself
 	// but this is a built-in TEDS column type that is used the the tabel viewer widget.
-	DSI->AddColumn<FTypedElementLabelColumn>(Row, { .Label = EntryData.Identifier.ToString() });
+	DSI->AddColumn<FTypedElementLabelColumn>(Row, { .Label = EntryData->Identifier.ToString() });
 
-	const RowHandle* ParentRow = EntryIndexToHandleMap.Find(EntryData.Parent);
+	const RowHandle* ParentRow = EntryIndexToHandleMap.Find(EntryData->Parent);
 	if (ParentRow)
 	{
 		DSI->AddColumn<FTableRowParentColumn>(Row, { .Parent = *ParentRow });
 	}
 
-	FHierarchyTableEditorModule& HierarchyTableModule = FModuleManager::GetModuleChecked<FHierarchyTableEditorModule>("HierarchyTableEditor");
-	const UHierarchyTableTypeHandler_Base* Handler = HierarchyTableModule.FindHandler(HierarchyTable->TableType);
-
-	for (const UScriptStruct* Column : Handler->GetColumns())
+	for (const UScriptStruct* Column : ElementHandler->GetColumns())
 	{
 		DSI->AddColumn(Row, Column);
 	}
diff --git a/Engine/Plugins/Experimental/Animation/HierarchyTable/Source/Editor/Private/HierarchyTableFactory.cpp b/Engine/Plugins/Experimental/Animation/HierarchyTable/Source/Editor/Private/HierarchyTableFactory.cpp
--- a/Engine/Plugins/Experimental/Animation/HierarchyTable/Source/Editor/Private/HierarchyTableFactory.cpp
+++ b/Engine/Plugins/Experimental/Animation/HierarchyTable/Source/Editor/Private/HierarchyTableFactory.cpp
@@ -2,72 +2,133 @@
 
 #include "HierarchyTableFactory.h"
 #include "HierarchyTable.h"
-#include "ContentBrowserModule.h"
-#include "IContentBrowserSingleton.h"
-#include "Editor.h"
 #include "Animation/Skeleton.h"
 #include "Widgets/SWindow.h"
-#include "Widgets/Layout/SBorder.h"
 #include "Kismet2/SClassPickerDialog.h"
 #include "StructViewerFilter.h"
 #include "StructViewerModule.h"
 #include "HierarchyTableEditorModule.h"
+#include "Modules/ModuleManager.h"
+#include "Widgets/Layout/SBorder.h"
+#include "Editor.h"
 
 UHierarchyTableFactory::UHierarchyTableFactory()
 {
 	SupportedClass = UHierarchyTable::StaticClass();
 	bCreateNew = true;
-	TableType = nullptr;
+	ElementType = nullptr;
 }
 
 UObject* UHierarchyTableFactory::FactoryCreateNew(UClass* Class, UObject* InParent, FName Name, EObjectFlags Flags, UObject* Context, FFeedbackContext* Warn)
 {
 	TObjectPtr<UHierarchyTable> HierarchyTable = NewObject<UHierarchyTable>(InParent, Class, Name, Flags, Context);
-	HierarchyTable->Skeleton = Skeleton.Get();
-	HierarchyTable->TableType = TableType;
+
+	HierarchyTable->TableMetadata = TableMetadata;
+	HierarchyTable->ElementType = ElementType;
 
 	FHierarchyTableEditorModule& HierarchyTableModule = FModuleManager::GetModuleChecked<FHierarchyTableEditorModule>("HierarchyTableEditor");
-	const UHierarchyTableTypeHandler_Base* Handler = HierarchyTableModule.FindHandler(TableType);
-	check(Handler);
+	const UHierarchyTable_TableTypeHandler* TableHandler = HierarchyTableModule.FindTableHandler(TableMetadata.GetScriptStruct());
+	check(TableHandler);
 
-	FInstancedStruct DefaultEntry = Handler->GetDefaultEntry();
-	HierarchyTable->InitializeTable(DefaultEntry);
+	TableHandler->InitialiseTable(TableMetadata, HierarchyTable);
 
 	return HierarchyTable;
 }
 
 bool UHierarchyTableFactory::ConfigureProperties()
 {
-	Skeleton = nullptr;
+	// Prompt the user to choose the table type
+	bool bSuccess = ConfigureTableType();
+	if (!bSuccess)
+	{
+		return false;
+	}
+
+	// Allow the table type to configure itself
+	{
+		FHierarchyTableEditorModule& HierarchyTableModule = FModuleManager::GetModuleChecked<FHierarchyTableEditorModule>("HierarchyTableEditor");
+		const UHierarchyTable_TableTypeHandler* TableHandler = HierarchyTableModule.FindTableHandler(TableMetadata.GetScriptStruct());
+		check(TableHandler);
+
+		bSuccess = TableHandler->FactoryConfigureProperties(TableMetadata);
+		if (!bSuccess)
+		{
+			return false;
+		}
+	}
+
+	// Prompt the user to choose the element type
+	bSuccess = ConfigureElementType();
+	if (!bSuccess)
+	{
+		return false;
+	}
+
+	return true;
+}
+
+bool UHierarchyTableFactory::ConfigureTableType()
+{
+	FStructViewerModule& StructViewerModule = FModuleManager::LoadModuleChecked<FStructViewerModule>("StructViewer");
 
-	FContentBrowserModule& ContentBrowserModule = FModuleManager::LoadModuleChecked<FContentBrowserModule>("ContentBrowser");
+	class FRuleFilter : public IStructViewerFilter
+	{
+	public:
+		FRuleFilter()
+		{
+		}
+
+		virtual bool IsStructAllowed(const FStructViewerInitializationOptions& InInitOptions, const UScriptStruct* InStruct, TSharedRef<FStructViewerFilterFuncs> InFilterFuncs) override
+		{
+			static const UScriptStruct* BaseStruct = FHierarchyTable_TableType::StaticStruct();
+			return InStruct != BaseStruct && InStruct->IsChildOf(BaseStruct);
+		}
 
-	FAssetPickerConfig AssetPickerConfig;
-	AssetPickerConfig.Filter.ClassPaths.Add(USkeleton::StaticClass()->GetClassPathName());
-	AssetPickerConfig.OnAssetSelected = FOnAssetSelected::CreateUObject(this, &UHierarchyTableFactory::OnSkeletonSelected);
-	AssetPickerConfig.InitialAssetViewType = EAssetViewType::List;
+		virtual bool IsUnloadedStructAllowed(const FStructViewerInitializationOptions& InInitOptions, const FSoftObjectPath& InStructPath, TSharedRef<FStructViewerFilterFuncs> InFilterFuncs) override
+		{
+			return false;
+		}
+	};
 
-	PickerWindow = SNew(SWindow)
-		.Title(INVTEXT("Pick Skeleton"))
+	static TSharedPtr<FRuleFilter> Filter = MakeShared<FRuleFilter>();
+	FStructViewerInitializationOptions Options;
+	{
+		Options.StructFilter = Filter;
+		Options.Mode = EStructViewerMode::StructPicker;
+		Options.DisplayMode = EStructViewerDisplayMode::ListView;
+		Options.NameTypeToDisplay = EStructViewerNameTypeToDisplay::DisplayName;
+		Options.bShowNoneOption = false;
+		Options.bShowUnloadedStructs = false;
+		Options.bAllowViewOptions = false;
+	}
+	
+	TSharedPtr<class SWindow> PickerWindow = SNew(SWindow)
+		.Title(INVTEXT("Pick Type"))
 		.ClientSize(FVector2D(500, 600))
 		.SupportsMinimize(false)
 		.SupportsMaximize(false)
 		[
 			SNew(SBorder)
 			.BorderImage(FAppStyle::GetBrush("Menu.Background"))
-			[
-				ContentBrowserModule.Get().CreateAssetPicker(AssetPickerConfig)
-			]
-		];
+		[
+			StructViewerModule.CreateStructViewer(Options, FOnStructPicked::CreateLambda([&](const UScriptStruct* ChosenStruct)
+				{
+					FInstancedStruct InstancedStruct;
+					InstancedStruct.InitializeAs(ChosenStruct);
+
+					TableMetadata = InstancedStruct;
+					PickerWindow->RequestDestroyWindow();
+				}))
+		]
+	];
 
 	GEditor->EditorAddModalWindow(PickerWindow.ToSharedRef());
-	PickerWindow.Reset();
 
-	if (Skeleton == nullptr)
-	{
-		return false;
-	}
+	return TableMetadata.IsValid();
+}
 
+bool UHierarchyTableFactory::ConfigureElementType()
+{
 	FStructViewerModule& StructViewerModule = FModuleManager::LoadModuleChecked<FStructViewerModule>("StructViewer");
 
 	class FRuleFilter : public IStructViewerFilter
@@ -79,7 +140,7 @@ bool UHierarchyTableFactory::ConfigureProperties()
 
 		virtual bool IsStructAllowed(const FStructViewerInitializationOptions& InInitOptions, const UScriptStruct* InStruct, TSharedRef<FStructViewerFilterFuncs> InFilterFuncs) override
 		{
-			static const UScriptStruct* BaseStruct = FHierarchyTableType::StaticStruct();
+			static const UScriptStruct* BaseStruct = FHierarchyTable_ElementType::StaticStruct();
 			return InStruct != BaseStruct && InStruct->IsChildOf(BaseStruct);
 		}
 
@@ -100,8 +161,8 @@ bool UHierarchyTableFactory::ConfigureProperties()
 		Options.bShowUnloadedStructs = false;
 		Options.bAllowViewOptions = false;
 	}
-
-	PickerWindow = SNew(SWindow)
+	
+	TSharedPtr<class SWindow> PickerWindow = SNew(SWindow)
 		.Title(INVTEXT("Pick Type"))
 		.ClientSize(FVector2D(500, 600))
 		.SupportsMinimize(false)
@@ -110,22 +171,15 @@ bool UHierarchyTableFactory::ConfigureProperties()
 			SNew(SBorder)
 			.BorderImage(FAppStyle::GetBrush("Menu.Background"))
 		[
-			StructViewerModule.CreateStructViewer(Options, FOnStructPicked::CreateLambda([this](const UScriptStruct* ChosenStruct)
+			StructViewerModule.CreateStructViewer(Options, FOnStructPicked::CreateLambda([&](const UScriptStruct* ChosenStruct)
 				{
-					TableType = ChosenStruct;
+					ElementType = ChosenStruct;
 					PickerWindow->RequestDestroyWindow();
 				}))
 		]
 	];
 
 	GEditor->EditorAddModalWindow(PickerWindow.ToSharedRef());
-	PickerWindow.Reset();
-
-	return TableType != nullptr;
-}
 
-void UHierarchyTableFactory::OnSkeletonSelected(const FAssetData& SelectedAsset)
-{
-	Skeleton = Cast<USkeleton>(SelectedAsset.GetAsset());
-	PickerWindow->RequestDestroyWindow();
+	return ElementType != nullptr;
 }
\ No newline at end of file
diff --git a/Engine/Plugins/Experimental/Animation/HierarchyTable/Source/Editor/Private/HierarchyTableLabelWidget.cpp b/Engine/Plugins/Experimental/Animation/HierarchyTable/Source/Editor/Private/HierarchyTableLabelWidget.cpp
--- a/Engine/Plugins/Experimental/Animation/HierarchyTable/Source/Editor/Private/HierarchyTableLabelWidget.cpp
+++ b/Engine/Plugins/Experimental/Animation/HierarchyTable/Source/Editor/Private/HierarchyTableLabelWidget.cpp
@@ -5,12 +5,13 @@
 #include "HierarchyTable/Columns/OverrideColumn.h"
 #include "Elements/Columns/TypedElementCompatibilityColumns.h"
 #include "Elements/Columns/TypedElementLabelColumns.h"
-#include "Elements/Framework/TypedElementAttributeBinding.h"
 #include "Widgets/Images/SImage.h"
 #include "Widgets/Layout/SSpacer.h"
 #include "Widgets/Text/STextBlock.h"
 #include "Widgets/SBoxPanel.h"
 #include "HierarchyTable.h"
+#include "HierarchyTableEditorModule.h"
+#include "Modules/ModuleManager.h"
 
 #define LOCTEXT_NAMESPACE "HierarchyTableLabelWidget"
 
@@ -41,60 +42,22 @@ TSharedPtr<SWidget> FHierarchyTableLabelWidgetConstructor::CreateWidget(
 	const UE::Editor::DataStorage::FMetaDataView& Arguments)
 {
 	using namespace UE::Editor::DataStorage;
-	if(DataStorage->IsRowAvailable(TargetRow))
-	{
-		FAttributeBinder Binder(TargetRow, DataStorage);
 
+	if (DataStorage->IsRowAvailable(TargetRow))
+	{
 		FTypedElementOverrideColumn* OverrideColumn = DataStorage->GetColumn<FTypedElementOverrideColumn>(TargetRow);
-		const FHierarchyTableEntryData& EntryData = OverrideColumn->OwnerTable->TableData[OverrideColumn->OwnerEntryIndex];
-
-		const FSlateBrush* LabelIcon = nullptr;
-		FSlateColor LabelIconColor = FSlateColor::UseForeground();
 
-		switch (EntryData.EntryType)
-		{
-			case EHierarchyTableEntryType::Bone:
-				LabelIcon = FAppStyle::GetBrush("SkeletonTree.Bone");
-				break;
-			case EHierarchyTableEntryType::Curve:
-				LabelIcon = FAppStyle::GetBrush("AnimGraph.Attribute.Curves.Icon");
-				LabelIconColor = FAppStyle::GetSlateColor("AnimGraph.Attribute.Curves.Color");
-				break;
-			case EHierarchyTableEntryType::Attribute:
-				LabelIcon = FAppStyle::GetBrush("AnimGraph.Attribute.Attributes.Icon");
-				LabelIconColor = FAppStyle::GetSlateColor("AnimGraph.Attribute.Attributes.Color");
-				break;
-		}
+		FHierarchyTableEditorModule& HierarchyTableModule = FModuleManager::GetModuleChecked<FHierarchyTableEditorModule>("HierarchyTableEditor");
+		const UHierarchyTable_TableTypeHandler* TableHandler = HierarchyTableModule.FindTableHandler(OverrideColumn->OwnerTable->TableMetadata.GetScriptStruct());
+		check(TableHandler);
 
-		return SNew(SHorizontalBox)
-			+SHorizontalBox::Slot()
-				.AutoWidth()
-			[
-				SNew(SImage)
-					.Image(LabelIcon)
-					.ColorAndOpacity(LabelIconColor)
-			]
-			+SHorizontalBox::Slot()
-				.AutoWidth()
-			[
-				SNew(SSpacer)
-					.Size(FVector2D(5.0f, 0.0f))
-			]
-			+SHorizontalBox::Slot()
-				.FillWidth(1.0f)
-			[
-				SNew(STextBlock)
-					.Text(Binder.BindText(&FTypedElementLabelColumn::Label))
-					.ToolTipText(Binder.BindText(&FTypedElementLabelColumn::Label))
-			];
+		return TableHandler->CreateLabelWidget(OverrideColumn->OwnerTable, OverrideColumn->OwnerEntryIndex);
 	}
 	else
 	{
 		return SNew(STextBlock)
 			.Text(LOCTEXT("MissingRowReferenceColumn", "Unable to retrieve row reference."));
 	}
-	
-	
 }
 
 #undef LOCTEXT_NAMESPACE
diff --git a/Engine/Plugins/Experimental/Animation/HierarchyTable/Source/Editor/Private/HierarchyTableTypeHandler.cpp b/Engine/Plugins/Experimental/Animation/HierarchyTable/Source/Editor/Private/HierarchyTableTypeHandler.cpp
new file mode 100644
--- /dev/null
+++ b/Engine/Plugins/Experimental/Animation/HierarchyTable/Source/Editor/Private/HierarchyTableTypeHandler.cpp
@@ -0,0 +1,21 @@
+// Copyright Epic Games, Inc. All Rights Reserved.
+
+#pragma once
+
+#include "HierarchyTableTypeHandler.h"
+#include "HierarchyTable.h"
+#include "Widgets/Images/SImage.h"
+#include "Widgets/Layout/SSpacer.h"
+#include "Widgets/Text/STextBlock.h"
+#include "Widgets/SBoxPanel.h"
+
+TSharedRef<SWidget> UHierarchyTable_TableTypeHandler::CreateLabelWidget(TObjectPtr<UHierarchyTable> HierarchyTable, const int32 EntryIndex) const
+{
+	return SNew(SHorizontalBox)
+		+SHorizontalBox::Slot()
+		.FillWidth(1.0f)
+		[
+			SNew(STextBlock)
+			.Text(FText::FromName(HierarchyTable->GetTableEntry(EntryIndex)->Identifier))
+		];
+}
\ No newline at end of file
diff --git a/Engine/Plugins/Experimental/Animation/HierarchyTable/Source/Editor/Public/HierarchyTableEditorModule.h b/Engine/Plugins/Experimental/Animation/HierarchyTable/Source/Editor/Public/HierarchyTableEditorModule.h
--- a/Engine/Plugins/Experimental/Animation/HierarchyTable/Source/Editor/Public/HierarchyTableEditorModule.h
+++ b/Engine/Plugins/Experimental/Animation/HierarchyTable/Source/Editor/Public/HierarchyTableEditorModule.h
@@ -13,12 +13,20 @@ public:
 	
 	virtual void ShutdownModule() override;
 
-	void HIERARCHYTABLEEDITOR_API RegisterTableType(const UScriptStruct* HierarchyTableType, const UHierarchyTableTypeHandler_Base* Handler);
+	HIERARCHYTABLEEDITOR_API void RegisterTableType(const UScriptStruct* TableType, const UHierarchyTable_TableTypeHandler* Handler);
 
-	void HIERARCHYTABLEEDITOR_API UnregisterTableType(const UScriptStruct* HierarchyTableType);
+	HIERARCHYTABLEEDITOR_API void UnregisterTableType(const UScriptStruct* TableType);
 
-	const UHierarchyTableTypeHandler_Base* FindHandler(const UScriptStruct* HierarchyTableType) const;
+	HIERARCHYTABLEEDITOR_API const UHierarchyTable_TableTypeHandler*  FindTableHandler(const UScriptStruct* TableType) const;
+
+	HIERARCHYTABLEEDITOR_API void RegisterElementType(const UScriptStruct* ElementType, const UHierarchyTable_ElementTypeHandler* Handler);
+
+	HIERARCHYTABLEEDITOR_API void UnregisterElementType(const UScriptStruct* ElementType);
+
+	HIERARCHYTABLEEDITOR_API const UHierarchyTable_ElementTypeHandler* FindElementHandler(const UScriptStruct* ElementType) const;
 
 private:
-	TMap<const UScriptStruct*, const UHierarchyTableTypeHandler_Base*> Handlers;
+	TMap<const UScriptStruct*, const UHierarchyTable_TableTypeHandler*> TableHandlers;
+
+	TMap<const UScriptStruct*, const UHierarchyTable_ElementTypeHandler*> ElementHandlers;
 };
\ No newline at end of file
diff --git a/Engine/Plugins/Experimental/Animation/HierarchyTable/Source/Editor/Public/HierarchyTableEditorToolkit.h b/Engine/Plugins/Experimental/Animation/HierarchyTable/Source/Editor/Public/HierarchyTableEditorToolkit.h
--- a/Engine/Plugins/Experimental/Animation/HierarchyTable/Source/Editor/Public/HierarchyTableEditorToolkit.h
+++ b/Engine/Plugins/Experimental/Animation/HierarchyTable/Source/Editor/Public/HierarchyTableEditorToolkit.h
@@ -6,8 +6,9 @@
 #include "HierarchyTable.h"
 #include "Toolkits/AssetEditorToolkit.h"
 #include "Elements/Common/TypedElementHandles.h"
+#include "IHierarchyTableEditor.h"
 
-class FHierarchyTableEditorToolkit : public FAssetEditorToolkit
+class FHierarchyTableEditorToolkit : public FAssetEditorToolkit, public IHierarchyTableEditor
 {
 public:
 	void InitEditor(const TArray<UObject*>& InObjects);
@@ -23,8 +24,12 @@ public:
 
 	void ExtendToolbar();
 
+	// Begin IHierarchyTableEditor
+	virtual int32 GetSelectedEntryIndex() const;
+	// End IHierarchyTableEditor
+
 private:
-	void AddEntry(const FName Identifier, const EHierarchyTableEntryType EntryType);
+	void HandleOnEntryAdded(const int32 EntryIndex);
 
 	TSharedRef<SWidget> CreateTedsOutliner();
 
@@ -35,4 +40,6 @@ private:
 	UE::Editor::DataStorage::QueryHandle InitialColumnQuery;
 
 	TSharedPtr<class ISceneOutliner> TedsOutlinerPtr;
+
+	FDelegateHandle OnEntryAddedHandle;
 };
diff --git a/Engine/Plugins/Experimental/Animation/HierarchyTable/Source/Editor/Public/HierarchyTableFactory.h b/Engine/Plugins/Experimental/Animation/HierarchyTable/Source/Editor/Public/HierarchyTableFactory.h
--- a/Engine/Plugins/Experimental/Animation/HierarchyTable/Source/Editor/Public/HierarchyTableFactory.h
+++ b/Engine/Plugins/Experimental/Animation/HierarchyTable/Source/Editor/Public/HierarchyTableFactory.h
@@ -5,6 +5,7 @@
 #include "CoreMinimal.h"
 #include "Factories/Factory.h"
 #include "HierarchyTableType.h"
+#include "StructUtils/InstancedStruct.h"
 
 #include "HierarchyTableFactory.generated.h"
 
@@ -15,11 +16,6 @@ class UHierarchyTableFactory : public UFactory
 {
 	GENERATED_BODY()
 
-	UPROPERTY()
-	TWeakObjectPtr<class USkeleton>	Skeleton;
-
-	const UScriptStruct* TableType;
-
 public:
 	UHierarchyTableFactory();
 
@@ -28,7 +24,11 @@ public:
 	virtual bool ConfigureProperties() override;
 
 private:
-	void OnSkeletonSelected(const FAssetData& SelectedAsset);
+	bool ConfigureTableType();
+
+	bool ConfigureElementType();
+
+	FInstancedStruct TableMetadata;
 
-	TSharedPtr<class SWindow> PickerWindow;
+	const UScriptStruct* ElementType;
 };
diff --git a/Engine/Plugins/Experimental/Animation/HierarchyTable/Source/Editor/Public/HierarchyTableTypeHandler.h b/Engine/Plugins/Experimental/Animation/HierarchyTable/Source/Editor/Public/HierarchyTableTypeHandler.h
--- a/Engine/Plugins/Experimental/Animation/HierarchyTable/Source/Editor/Public/HierarchyTableTypeHandler.h
+++ b/Engine/Plugins/Experimental/Animation/HierarchyTable/Source/Editor/Public/HierarchyTableTypeHandler.h
@@ -4,20 +4,41 @@
 
 #include "UObject/Object.h"
 #include "StructUtils/InstancedStruct.h"
+#include "HierarchyTableType.h"
 #include "HierarchyTableTypeHandler.generated.h"
 
 class UHierarchyTable;
 class UScriptStruct;
+class UToolMenu;
+class SWidget;
+class IHierarchyTableEditor;
 
 UCLASS(Abstract)
-class HIERARCHYTABLEEDITOR_API UHierarchyTableTypeHandler_Base : public UObject
+class HIERARCHYTABLEEDITOR_API UHierarchyTable_TableTypeHandler : public UObject
 {
 	GENERATED_BODY()
 
 public:
-	virtual FInstancedStruct GetDefaultEntry() const { return FInstancedStruct(); }
+	virtual TSharedPtr<FHierarchyTable_TableType> CreateTableType() const { return nullptr; }
+
+	virtual void InitialiseTable(const FInstancedStruct& TableType, TObjectPtr<UHierarchyTable> HierarchyTable) const {}
+
+	virtual bool FactoryConfigureProperties(FInstancedStruct& TableType) const { return false; }
+
+	//virtual int32 AddEntry(TObjectPtr<UHierarchyTable> HierarchyTable, const int32 ParentIndex) const { return false; }
+
+	virtual void ExtendToolbar(UToolMenu* ToolMenu, UHierarchyTable* HierarchyTable, const IHierarchyTableEditor* Editor) const {}
+
+	virtual TSharedRef<SWidget> CreateLabelWidget(TObjectPtr<UHierarchyTable> HierarchyTable, const int32 EntryIndex) const;
+};
 
-	virtual void InitializeTable(TObjectPtr<UHierarchyTable> InHierarchyTable) const {};
+UCLASS(Abstract)
+class HIERARCHYTABLEEDITOR_API UHierarchyTable_ElementTypeHandler : public UObject
+{
+	GENERATED_BODY()
+
+public:
+	virtual FInstancedStruct GetDefaultEntry() const { return FInstancedStruct(); }
 
-	virtual TArray<UScriptStruct*> GetColumns() const { return {}; };
+	virtual TArray<UScriptStruct*> GetColumns() const { return {}; }
 };
\ No newline at end of file
diff --git a/Engine/Plugins/Experimental/Animation/HierarchyTable/Source/Editor/Public/IHierarchyTableEditor.h b/Engine/Plugins/Experimental/Animation/HierarchyTable/Source/Editor/Public/IHierarchyTableEditor.h
new file mode 100644
--- /dev/null
+++ b/Engine/Plugins/Experimental/Animation/HierarchyTable/Source/Editor/Public/IHierarchyTableEditor.h
@@ -0,0 +1,11 @@
+// Copyright Epic Games, Inc. All Rights Reserved.
+
+#pragma once
+
+#include "CoreMinimal.h"
+
+class IHierarchyTableEditor
+{
+public:
+	virtual int32 GetSelectedEntryIndex() const = 0;
+};
\ No newline at end of file
diff --git a/Engine/Plugins/Experimental/Animation/HierarchyTable/Source/Runtime/HierarchyTableRuntime.Build.cs b/Engine/Plugins/Experimental/Animation/HierarchyTable/Source/Runtime/HierarchyTableRuntime.Build.cs
--- a/Engine/Plugins/Experimental/Animation/HierarchyTable/Source/Runtime/HierarchyTableRuntime.Build.cs
+++ b/Engine/Plugins/Experimental/Animation/HierarchyTable/Source/Runtime/HierarchyTableRuntime.Build.cs
@@ -11,8 +11,7 @@ public class HierarchyTableRuntime : ModuleRules
 			{
 				"Core",
 				"CoreUObject",
-				"Engine",
-				"AnimationCore"
+				"Engine"
 			}
 		);
 
diff --git a/Engine/Plugins/Experimental/Animation/HierarchyTable/Source/Runtime/Private/HierarchyTable.cpp b/Engine/Plugins/Experimental/Animation/HierarchyTable/Source/Runtime/Private/HierarchyTable.cpp
--- a/Engine/Plugins/Experimental/Animation/HierarchyTable/Source/Runtime/Private/HierarchyTable.cpp
+++ b/Engine/Plugins/Experimental/Animation/HierarchyTable/Source/Runtime/Private/HierarchyTable.cpp
@@ -5,26 +5,7 @@
 #include "ReferenceSkeleton.h"
 
 UHierarchyTable::UHierarchyTable()
-{}
-
-void UHierarchyTable::InitializeTable(FInstancedStruct DefaultEntry)
 {
-	check(Skeleton);
-	check(TableData.Num() == 0);
-
-	FReferenceSkeleton RefSkeleton = Skeleton->GetReferenceSkeleton();
-
-	for (int32 BoneIndex = 0; BoneIndex < RefSkeleton.GetNum(); ++BoneIndex)
-	{
-		FHierarchyTableEntryData EntryData;
-		EntryData.Parent = RefSkeleton.GetParentIndex(BoneIndex);
-		EntryData.Identifier = RefSkeleton.GetBoneName(BoneIndex);
-		EntryData.EntryType = EHierarchyTableEntryType::Bone;
-		EntryData.Payload = BoneIndex == 0 ? DefaultEntry : TOptional<FInstancedStruct>();
-		EntryData.OwnerTable = this;
-
-		TableData.Add(EntryData);
-	}
 }
 
 bool FHierarchyTableEntryData::HasOverriddenChildren() const
@@ -68,17 +49,17 @@ const FInstancedStruct* FHierarchyTableEntryData::GetActualValue() const
 	return IsOverridden() ? Payload.GetPtrOrNull() : GetFromClosestAncestor();
 }
 
-const FHierarchyTableEntryData* FHierarchyTableEntryData::GetClosestAncestor()
+const FHierarchyTableEntryData* FHierarchyTableEntryData::GetClosestAncestor() const
 {
-	return IsOverridden() ? this : OwnerTable->TableData[Parent].GetClosestAncestor();
+	return IsOverridden() ? this : OwnerTable->GetTableEntry(Parent)->GetClosestAncestor();
 }
 
 const FInstancedStruct* FHierarchyTableEntryData::GetFromClosestAncestor() const
 {
 	check(Parent != INDEX_NONE);
 
-	const FHierarchyTableEntryData& ParentEntry = OwnerTable->TableData[Parent];
-	return ParentEntry.GetActualValue();
+	const FHierarchyTableEntryData* const ParentEntry = OwnerTable->GetTableEntry(Parent);
+	return ParentEntry->GetActualValue();
 }
 
 TArray<const FHierarchyTableEntryData*> UHierarchyTable::GetChildren(const FHierarchyTableEntryData& Parent) const
@@ -114,10 +95,41 @@ bool UHierarchyTable::HasIdentifier(const FName Identifier) const
 		});
 }
 
-FHierarchyTableEntryData* UHierarchyTable::FindEntry(const FName EntryIdentifier, const EHierarchyTableEntryType EntryType)
+int32 UHierarchyTable::AddEntry(const FHierarchyTableEntryData& Entry)
 {
-	return TableData.FindByPredicate([EntryIdentifier, EntryType](const FHierarchyTableEntryData& Entry)
-		{
-			return Entry.Identifier == EntryIdentifier && Entry.EntryType == EntryType;
-		});
+	const int32 EntryIndex = TableData.Add(Entry);
+#if WITH_EDITOR
+	OnEntryAdded.Broadcast(EntryIndex);
+#endif //WITH_EDITOR
+	return EntryIndex;
+}
+
+#if WITH_EDITOR
+FDelegateHandle UHierarchyTable::RegisterOnEntryAdded(const FOnEntryAdded::FDelegate& Delegate)
+{
+	return OnEntryAdded.Add(Delegate);
+}
+
+bool UHierarchyTable::UnregisterOnEntryAdded(const FDelegateHandle& Handle)
+{
+	return OnEntryAdded.Remove(Handle);
+}
+
+const FHierarchyTableEntryData* const UHierarchyTable::GetTableEntry(const int32 EntryIndex) const
+{
+	if (TableData.IsValidIndex(EntryIndex))
+	{
+		return &TableData[EntryIndex];
+	}
+	return nullptr;
+}
+#endif //WITH_EDITOR
+
+FHierarchyTableEntryData* const UHierarchyTable::GetMutableTableEntry(const int32 EntryIndex)
+{
+	if (TableData.IsValidIndex(EntryIndex))
+	{
+		return &TableData[EntryIndex];
+	}
+	return nullptr;
 }
\ No newline at end of file
diff --git a/Engine/Plugins/Experimental/Animation/HierarchyTable/Source/Runtime/Public/DefaultHierarchyTableType.h b/Engine/Plugins/Experimental/Animation/HierarchyTable/Source/Runtime/Public/DefaultHierarchyTableType.h
new file mode 100644
--- /dev/null
+++ b/Engine/Plugins/Experimental/Animation/HierarchyTable/Source/Runtime/Public/DefaultHierarchyTableType.h
@@ -0,0 +1,13 @@
+// Copyright Epic Games, Inc. All Rights Reserved.
+
+#pragma once
+
+#include "HierarchyTableType.h"
+
+#include "DefaultHierarchyTableType.generated.h"
+
+USTRUCT()
+struct FHierarchyTable_TableType_Default final : public FHierarchyTable_TableType
+{
+	GENERATED_BODY()
+};
\ No newline at end of file
diff --git a/Engine/Plugins/Experimental/Animation/HierarchyTable/Source/Runtime/Public/HierarchyTable.h b/Engine/Plugins/Experimental/Animation/HierarchyTable/Source/Runtime/Public/HierarchyTable.h
--- a/Engine/Plugins/Experimental/Animation/HierarchyTable/Source/Runtime/Public/HierarchyTable.h
+++ b/Engine/Plugins/Experimental/Animation/HierarchyTable/Source/Runtime/Public/HierarchyTable.h
@@ -3,26 +3,17 @@
 #pragma once
 
 #include "CoreMinimal.h"
-#include "UObject/NoExportTypes.h"
 #include "HierarchyTableType.h"
-#include "Templates/SubclassOf.h"
-#include "Templates/SharedPointer.h"
 #include "StructUtils/InstancedStruct.h"
 
 #include "HierarchyTable.generated.h"
 
 class USkeleton;
-
-UENUM()
-enum class EHierarchyTableEntryType : uint8
-{
-	Bone,
-	Curve,
-	Attribute
-};
-
 class UHierarchyTable;
 
+/**
+ * The data associated with each item in a hierarchy table.
+ */
 USTRUCT()
 struct HIERARCHYTABLERUNTIME_API FHierarchyTableEntryData
 {
@@ -30,27 +21,35 @@ struct HIERARCHYTABLERUNTIME_API FHierarchyTableEntryData
 
 public:
 	FHierarchyTableEntryData()
-		: EntryType(EHierarchyTableEntryType::Bone)
-		, Parent(INDEX_NONE)
+		: Parent(INDEX_NONE)
 	{
 	}
 
 	friend UHierarchyTable;
 
+	// The hierarchy table that this entry resides in.
 	UPROPERTY()
 	TObjectPtr<UHierarchyTable> OwnerTable;
 
+	/**
+	 * The actual user set data this element stores, matches the owner table's element type.
+	 * If unset, will inherit from the parent.
+	 * Is always set for entries with no parent.
+	 */
 	UPROPERTY()
-	EHierarchyTableEntryType EntryType;
+	TOptional<FInstancedStruct> Payload;
 
+	// This entry's display name and unique identifier.
 	UPROPERTY()
 	FName Identifier;
 
+	// The index of this entry's parent, is INDEX_NONE in the case of a root element.
 	UPROPERTY()
 	int32 Parent;
 
+	// A second data payload specific to the table's type and is read-only metadata.
 	UPROPERTY()
-	TOptional<FInstancedStruct> Payload;
+	FInstancedStruct TablePayload;
 
 public:
 	bool HasParent() const { return Parent != INDEX_NONE; }
@@ -74,7 +73,7 @@ public:
 		return Payload.GetValue().GetMutablePtr<T>();
 	}
 
-	const FHierarchyTableEntryData* GetClosestAncestor();
+	const FHierarchyTableEntryData* GetClosestAncestor() const;
 
 private:
 	// TODO: This needs to be cached in some way as it potentially goes up the entire hierarchy until it finds an ancestor with an overridden value
@@ -87,6 +86,9 @@ private:
 	bool IsOverriddenOrHasOverriddenChildren(const bool bIncludeSelf) const;
 };
 
+/**
+ * A general-purpose container asset for storing typed hierarchical data.
+ */
 UCLASS()
 class HIERARCHYTABLERUNTIME_API UHierarchyTable : public UObject
 {
@@ -95,20 +97,48 @@ class HIERARCHYTABLERUNTIME_API UHierarchyTable : public UObject
 public:
 	UHierarchyTable();
 
-	UPROPERTY(VisibleAnywhere, Category=Default)
-	TObjectPtr<USkeleton> Skeleton;
+	// The table metadata stores any data dependent on the table type needed to create and maintain the hierarchy.
+	UPROPERTY(VisibleAnywhere, Category = Default)
+	FInstancedStruct TableMetadata;
 
+	// The element type is the type that each element is mapped into.
 	UPROPERTY(VisibleAnywhere, Category = Default)
-	TObjectPtr<const UScriptStruct> TableType;
+	TObjectPtr<const UScriptStruct> ElementType;
 
+private:
+	// The actual table data
 	UPROPERTY()
 	TArray<FHierarchyTableEntryData> TableData;
 
-	void InitializeTable(FInstancedStruct DefaultEntry);
+public:
+	const TArray<FHierarchyTableEntryData>& GetTableData() const { return TableData; }
+
+	const FHierarchyTableEntryData* const GetTableEntry(const int32 EntryIndex) const;
+
+	FHierarchyTableEntryData* const GetMutableTableEntry(const int32 EntryIndex);
+
+	int32 AddEntry(const FHierarchyTableEntryData& Entry);
 
 	TArray<const FHierarchyTableEntryData*> GetChildren(const FHierarchyTableEntryData& Parent) const;
 
 	bool HasIdentifier(const FName Identifier) const;
 
-	FHierarchyTableEntryData* FindEntry(const FName EntryIdentifier, const EHierarchyTableEntryType EntryType);
+	template <typename T>
+	FHierarchyTableEntryData* FindEntry(const FName EntryIdentifier)
+	{
+		return TableData.FindByPredicate([EntryIdentifier](const FHierarchyTableEntryData& Entry)
+		{
+			return Entry.Identifier == EntryIdentifier;
+		});
+	}
+
+#if WITH_EDITOR
+	DECLARE_MULTICAST_DELEGATE_OneParam(FOnEntryAdded, int32);
+private:
+	FOnEntryAdded OnEntryAdded;
+
+public:
+	FDelegateHandle RegisterOnEntryAdded(const FOnEntryAdded::FDelegate& Delegate);
+	bool UnregisterOnEntryAdded(const FDelegateHandle& Handle);
+#endif
 };
diff --git a/Engine/Plugins/Experimental/Animation/HierarchyTable/Source/Runtime/Public/HierarchyTableType.h b/Engine/Plugins/Experimental/Animation/HierarchyTable/Source/Runtime/Public/HierarchyTableType.h
--- a/Engine/Plugins/Experimental/Animation/HierarchyTable/Source/Runtime/Public/HierarchyTableType.h
+++ b/Engine/Plugins/Experimental/Animation/HierarchyTable/Source/Runtime/Public/HierarchyTableType.h
@@ -5,7 +5,19 @@
 #include "HierarchyTableType.generated.h"
 
 USTRUCT()
-struct HIERARCHYTABLERUNTIME_API FHierarchyTableType
+struct HIERARCHYTABLERUNTIME_API FHierarchyTable_TableType
+{
+	GENERATED_BODY()
+};
+
+USTRUCT()
+struct HIERARCHYTABLERUNTIME_API FHierarchyTable_TablePayloadType
+{
+	GENERATED_BODY()
+};
+
+USTRUCT()
+struct HIERARCHYTABLERUNTIME_API FHierarchyTable_ElementType
 {
 	GENERATED_BODY()
 };
\ No newline at end of file
diff --git a/Engine/Plugins/Experimental/Animation/HierarchyTableBuiltin/HierarchyTableBuiltin.uplugin b/Engine/Plugins/Experimental/Animation/HierarchyTableBuiltin/HierarchyTableBuiltin.uplugin
--- a/Engine/Plugins/Experimental/Animation/HierarchyTableBuiltin/HierarchyTableBuiltin.uplugin
+++ b/Engine/Plugins/Experimental/Animation/HierarchyTableBuiltin/HierarchyTableBuiltin.uplugin
@@ -2,8 +2,8 @@
 	"FileVersion": 3,
 	"Version": 1,
 	"VersionName": "0.1",
-	"FriendlyName": "Hierarchy Table Builtin",
-	"Description": "Contains built-in types for controlling animation blending",
+	"FriendlyName": "Hierarchy Table Animation",
+	"Description": "Animation-specific type definitions for Hierarchy Tables",
 	"Category": "Animation",
 	"CreatedBy": "Epic Games, Inc.",
 	"CreatedByURL": "https://epicgames.com",
diff --git a/Engine/Plugins/Experimental/Animation/HierarchyTableBuiltin/Source/Editor/HierarchyTableBuiltinEditor.Build.cs b/Engine/Plugins/Experimental/Animation/HierarchyTableBuiltin/Source/Editor/HierarchyTableBuiltinEditor.Build.cs
--- a/Engine/Plugins/Experimental/Animation/HierarchyTableBuiltin/Source/Editor/HierarchyTableBuiltinEditor.Build.cs
+++ b/Engine/Plugins/Experimental/Animation/HierarchyTableBuiltin/Source/Editor/HierarchyTableBuiltinEditor.Build.cs
@@ -18,7 +18,8 @@ public class HierarchyTableBuiltinEditor : ModuleRules
 			"HierarchyTableRuntime",
 			"HierarchyTableEditor",
 			"HierarchyTableBuiltinRuntime",
-			"TypedElementFramework" // TODO: Remove
+			"TypedElementFramework", // TODO: Remove
+			"ToolMenus"
 		);
 	}
 }
diff --git a/Engine/Plugins/Experimental/Animation/HierarchyTableBuiltin/Source/Editor/Private/HierachyTableBuiltinEditorModule.cpp b/Engine/Plugins/Experimental/Animation/HierarchyTableBuiltin/Source/Editor/Private/HierachyTableBuiltinEditorModule.cpp
--- a/Engine/Plugins/Experimental/Animation/HierarchyTableBuiltin/Source/Editor/Private/HierachyTableBuiltinEditorModule.cpp
+++ b/Engine/Plugins/Experimental/Animation/HierarchyTableBuiltin/Source/Editor/Private/HierachyTableBuiltinEditorModule.cpp
@@ -7,12 +7,16 @@
 #include "MaskProfile/MaskProfileTypeHandler.h"
 #include "Modules/ModuleManager.h"
 #include "HierarchyTableEditorModule.h"
+#include "SkeletonHierarchyTableType.h"
+#include "SkeletonHierarchyTableTypeHandler.h"
 
 void FHierarchyTableBuiltinEditorModule::StartupModule()
 {
 	FHierarchyTableEditorModule& HierarchyTableModule = FModuleManager::LoadModuleChecked<FHierarchyTableEditorModule>("HierarchyTableEditor");
-	HierarchyTableModule.RegisterTableType(FHierarchyTableType_Mask::StaticStruct(), GetDefault<UHierarchyTableTypeHandler_Mask>());
-	HierarchyTableModule.RegisterTableType(FHierarchyTableType_Time::StaticStruct(), GetDefault<UHierarchyTableTypeHandler_Time>());
+	HierarchyTableModule.RegisterElementType(FHierarchyTable_ElementType_Mask::StaticStruct(), GetDefault<UHierarchyTable_ElementTypeHandler_Mask>());
+	HierarchyTableModule.RegisterElementType(FHierarchyTable_ElementType_Time::StaticStruct(), GetDefault<UHierarchyTable_ElementTypeHandler_Time>());
+
+	HierarchyTableModule.RegisterTableType(FHierarchyTable_TableType_Skeleton::StaticStruct(), GetDefault<UHierarchyTable_TableTypeHandler_Skeleton>());
 }
 
 void FHierarchyTableBuiltinEditorModule::ShutdownModule()
@@ -20,8 +24,10 @@ void FHierarchyTableBuiltinEditorModule::ShutdownModule()
 	if (FModuleManager::Get().IsModuleLoaded("HierarchyTableEditor"))
 	{
 		FHierarchyTableEditorModule& HierarchyTableModule = FModuleManager::GetModuleChecked<FHierarchyTableEditorModule>("HierarchyTableEditor");
-		HierarchyTableModule.UnregisterTableType(FHierarchyTableType_Mask::StaticStruct());
-		HierarchyTableModule.UnregisterTableType(FHierarchyTableType_Time::StaticStruct());
+		HierarchyTableModule.UnregisterElementType(FHierarchyTable_ElementType_Mask::StaticStruct());
+		HierarchyTableModule.UnregisterElementType(FHierarchyTable_ElementType_Time::StaticStruct());
+
+		HierarchyTableModule.UnregisterTableType(FHierarchyTable_TableType_Skeleton::StaticStruct());
 	}
 }
 
diff --git a/Engine/Plugins/Experimental/Animation/HierarchyTableBuiltin/Source/Editor/Private/MaskProfile/MaskProfileTypeHandler.cpp b/Engine/Plugins/Experimental/Animation/HierarchyTableBuiltin/Source/Editor/Private/MaskProfile/MaskProfileTypeHandler.cpp
--- a/Engine/Plugins/Experimental/Animation/HierarchyTableBuiltin/Source/Editor/Private/MaskProfile/MaskProfileTypeHandler.cpp
+++ b/Engine/Plugins/Experimental/Animation/HierarchyTableBuiltin/Source/Editor/Private/MaskProfile/MaskProfileTypeHandler.cpp
@@ -4,17 +4,17 @@
 #include "MaskProfile/HierarchyTableTypeMask.h"
 #include "MaskProfile/MaskProfileProxyColumn.h"
 
-FInstancedStruct UHierarchyTableTypeHandler_Mask::GetDefaultEntry() const
+FInstancedStruct UHierarchyTable_ElementTypeHandler_Mask::GetDefaultEntry() const
 {
 	FInstancedStruct NewEntry;
-	NewEntry.InitializeAs(FHierarchyTableType_Mask::StaticStruct());
-	FHierarchyTableType_Mask& NewEntryRef = NewEntry.GetMutable<FHierarchyTableType_Mask>();
+	NewEntry.InitializeAs(FHierarchyTable_ElementType_Mask::StaticStruct());
+	FHierarchyTable_ElementType_Mask& NewEntryRef = NewEntry.GetMutable<FHierarchyTable_ElementType_Mask>();
 	NewEntryRef.Value = 1.0f;
 
 	return NewEntry;
 }
 
-TArray<UScriptStruct*> UHierarchyTableTypeHandler_Mask::GetColumns() const
+TArray<UScriptStruct*> UHierarchyTable_ElementTypeHandler_Mask::GetColumns() const
 {
 	return { FHierarchyTableMaskColumn_Value::StaticStruct() };
 }
\ No newline at end of file
diff --git a/Engine/Plugins/Experimental/Animation/HierarchyTableBuiltin/Source/Editor/Private/MaskProfile/MaskProfileTypeHandler.h b/Engine/Plugins/Experimental/Animation/HierarchyTableBuiltin/Source/Editor/Private/MaskProfile/MaskProfileTypeHandler.h
--- a/Engine/Plugins/Experimental/Animation/HierarchyTableBuiltin/Source/Editor/Private/MaskProfile/MaskProfileTypeHandler.h
+++ b/Engine/Plugins/Experimental/Animation/HierarchyTableBuiltin/Source/Editor/Private/MaskProfile/MaskProfileTypeHandler.h
@@ -7,7 +7,7 @@
 #include "MaskProfileTypeHandler.generated.h"
 
 UCLASS()
-class UHierarchyTableTypeHandler_Mask final : public UHierarchyTableTypeHandler_Base
+class UHierarchyTable_ElementTypeHandler_Mask final : public UHierarchyTable_ElementTypeHandler
 {
 	GENERATED_BODY()
 
diff --git a/Engine/Plugins/Experimental/Animation/HierarchyTableBuiltin/Source/Editor/Private/MaskProfile/MaskProfileWidgetConstructor.cpp b/Engine/Plugins/Experimental/Animation/HierarchyTableBuiltin/Source/Editor/Private/MaskProfile/MaskProfileWidgetConstructor.cpp
--- a/Engine/Plugins/Experimental/Animation/HierarchyTableBuiltin/Source/Editor/Private/MaskProfile/MaskProfileWidgetConstructor.cpp
+++ b/Engine/Plugins/Experimental/Animation/HierarchyTableBuiltin/Source/Editor/Private/MaskProfile/MaskProfileWidgetConstructor.cpp
@@ -11,17 +11,17 @@
 TSharedRef<SWidget> FHierarchyTableMaskWidgetConstructor_Value::CreateInternalWidget(UHierarchyTable* HierarchyTable, int32 EntryIndex)
 {
 	return SNew(SSpinBox<float>)
-		.IsEnabled_Lambda([HierarchyTable, EntryIndex]() { return HierarchyTable->TableData[EntryIndex].IsOverridden(); })
+		.IsEnabled_Lambda([HierarchyTable, EntryIndex]() { return HierarchyTable->GetTableEntry(EntryIndex)->IsOverridden(); })
 		.MinDesiredWidth(100.0f)
 		.MinValue(0.0f)
 		.MaxValue(1.0f)
 		.Value_Lambda([HierarchyTable, EntryIndex]()
 			{
-				return HierarchyTable->TableData[EntryIndex].GetValue<FHierarchyTableType_Mask>()->Value;
+				return HierarchyTable->GetTableEntry(EntryIndex)->GetValue<FHierarchyTable_ElementType_Mask>()->Value;
 			})
 		.OnValueChanged_Lambda([HierarchyTable, EntryIndex](float NewValue)
 			{
-				HierarchyTable->TableData[EntryIndex].GetMutableValue<FHierarchyTableType_Mask>()->Value = NewValue;
+				HierarchyTable->GetMutableTableEntry(EntryIndex)->GetMutableValue<FHierarchyTable_ElementType_Mask>()->Value = NewValue;
 			})
 		.OnBeginSliderMovement_Lambda([HierarchyTable]()
 			{
diff --git a/Engine/Plugins/Experimental/Animation/HierarchyTableBuiltin/Source/Editor/Private/SkeletonHierarchyTableTypeHandler.cpp b/Engine/Plugins/Experimental/Animation/HierarchyTableBuiltin/Source/Editor/Private/SkeletonHierarchyTableTypeHandler.cpp
new file mode 100644
--- /dev/null
+++ b/Engine/Plugins/Experimental/Animation/HierarchyTableBuiltin/Source/Editor/Private/SkeletonHierarchyTableTypeHandler.cpp
@@ -0,0 +1,298 @@
+// Copyright Epic Games, Inc. All Rights Reserved.
+
+#include "SkeletonHierarchyTableTypeHandler.h"
+#include "ContentBrowserModule.h"
+#include "IContentBrowserSingleton.h"
+#include "SkeletonHierarchyTableType.h"
+#include "Animation/Skeleton.h"
+#include "Widgets/Layout/SBorder.h"
+#include "Editor.h"
+#include "StructUtils/InstancedStruct.h"
+#include "ToolMenu.h"
+#include "ToolMenuMisc.h"
+#include "Framework/Application/SlateApplication.h"
+#include "Widgets/Input/STextEntryPopup.h"
+#include "PersonaModule.h"
+#include "HierarchyTable.h"
+#include "HierarchyTableEditorModule.h"
+#include "Modules/ModuleManager.h"
+#include "Widgets/Images/SImage.h"
+#include "Widgets/Layout/SSpacer.h"
+#include "Widgets/Text/STextBlock.h"
+#include "Widgets/SBoxPanel.h"
+#include "IHierarchyTableEditor.h"
+
+#define LOCTEXT_NAMESPACE "UHierarchyTable_TableTypeHandler_Skeleton"
+
+TSharedPtr<FHierarchyTable_TableType> UHierarchyTable_TableTypeHandler_Skeleton::CreateTableType() const
+{
+	return MakeShared<FHierarchyTable_TableType_Skeleton>();
+}
+
+void UHierarchyTable_TableTypeHandler_Skeleton::InitialiseTable(const FInstancedStruct& TableType, TObjectPtr<UHierarchyTable> HierarchyTable) const
+{
+	const FHierarchyTable_TableType_Skeleton& SkeletonTableType = TableType.Get<FHierarchyTable_TableType_Skeleton>();
+
+	check(SkeletonTableType.Skeleton);
+	check(HierarchyTable->GetTableData().Num() == 0);
+
+	FHierarchyTableEditorModule& HierarchyTableModule = FModuleManager::GetModuleChecked<FHierarchyTableEditorModule>("HierarchyTableEditor");
+	const UHierarchyTable_ElementTypeHandler* ElementHandler = HierarchyTableModule.FindElementHandler(HierarchyTable->ElementType);
+	check(ElementHandler);
+
+	FInstancedStruct DefaultEntry = ElementHandler->GetDefaultEntry();
+
+	FInstancedStruct DefaultTablePayload;
+	DefaultTablePayload.InitializeAs<FHierarchyTable_TablePayloadType_Skeleton>();
+	DefaultTablePayload.GetMutable<FHierarchyTable_TablePayloadType_Skeleton>().EntryType = ESkeletonHierarchyTable_TablePayloadEntryType::Bone;
+
+	FReferenceSkeleton RefSkeleton = SkeletonTableType.Skeleton->GetReferenceSkeleton();
+
+	for (int32 BoneIndex = 0; BoneIndex < RefSkeleton.GetNum(); ++BoneIndex)
+	{
+		FHierarchyTableEntryData EntryData;
+		EntryData.Parent = RefSkeleton.GetParentIndex(BoneIndex);
+		EntryData.Identifier = RefSkeleton.GetBoneName(BoneIndex);
+		EntryData.TablePayload = DefaultTablePayload;
+		EntryData.Payload = BoneIndex == 0 ? DefaultEntry : TOptional<FInstancedStruct>();
+		EntryData.OwnerTable = HierarchyTable;
+		HierarchyTable->AddEntry(EntryData);
+	}
+}
+
+bool UHierarchyTable_TableTypeHandler_Skeleton::FactoryConfigureProperties(FInstancedStruct& TableType) const
+{
+	FHierarchyTable_TableType_Skeleton& SkeletonTableType = TableType.GetMutable<FHierarchyTable_TableType_Skeleton>();
+
+	FContentBrowserModule& ContentBrowserModule = FModuleManager::LoadModuleChecked<FContentBrowserModule>("ContentBrowser");
+
+	TSharedPtr<class SWindow> PickerWindow;
+
+	FAssetPickerConfig AssetPickerConfig;
+	AssetPickerConfig.Filter.ClassPaths.Add(USkeleton::StaticClass()->GetClassPathName());
+	AssetPickerConfig.InitialAssetViewType = EAssetViewType::List;
+	AssetPickerConfig.OnAssetSelected = FOnAssetSelected::CreateLambda([&](const FAssetData& SelectedAsset)
+		{
+			SkeletonTableType.Skeleton = Cast<USkeleton>(SelectedAsset.GetAsset());
+			PickerWindow->RequestDestroyWindow();
+		});
+
+	PickerWindow = SNew(SWindow)
+		.Title(INVTEXT("Pick Skeleton"))
+		.ClientSize(FVector2D(500, 600))
+		.SupportsMinimize(false)
+		.SupportsMaximize(false)
+		[
+			SNew(SBorder)
+			.BorderImage(FAppStyle::GetBrush("Menu.Background"))
+			[
+				ContentBrowserModule.Get().CreateAssetPicker(AssetPickerConfig)
+			]
+		];
+
+	GEditor->EditorAddModalWindow(PickerWindow.ToSharedRef());
+	PickerWindow.Reset();
+
+	return SkeletonTableType.Skeleton != nullptr;
+}
+
+void UHierarchyTable_TableTypeHandler_Skeleton::ExtendToolbar(UToolMenu* ToolMenu, UHierarchyTable* HierarchyTable, const IHierarchyTableEditor* Editor) const
+{
+	const FToolMenuInsert SectionInsertLocation("Asset", EToolMenuInsertType::After);
+	{
+		FToolMenuSection& HierarchyTableSection = ToolMenu->AddSection("HierarchyTable", LOCTEXT("HierarchyTable_ToolbarLabel", "HierarchyTable"), SectionInsertLocation);
+
+		HierarchyTableSection.AddEntry(FToolMenuEntry::InitComboButton(
+			"AddCurve",
+			FUIAction(),
+			FNewToolMenuDelegate::CreateLambda([=, this](UToolMenu* InSubMenu)
+				{
+					FToolUIAction Action;
+				    Action.ExecuteAction = FToolMenuExecuteAction::CreateLambda([=, this](const FToolMenuContext& Context)
+				    {
+						TSharedRef<STextEntryPopup> TextEntry = SNew(STextEntryPopup)
+							.Label(LOCTEXT("NewCurveEntryLabal", "Curve Name"))
+							.OnTextCommitted_Lambda([=, this](const FText& CommittedText, ETextCommit::Type CommitInfo)
+								{
+									const int32 SelectedIndex = Editor->GetSelectedEntryIndex();
+									const int32 ParentIndex = SelectedIndex == INDEX_NONE ? 0 : SelectedIndex;
+									AddCurve(HierarchyTable, ParentIndex, FName(CommittedText.ToString()));
+									FSlateApplication::Get().DismissAllMenus();
+								});
+
+						FSlateApplication& SlateApp = FSlateApplication::Get();
+						SlateApp.PushMenu(
+							SlateApp.GetInteractiveTopLevelWindows()[0],
+							FWidgetPath(),
+							TextEntry,
+							SlateApp.GetCursorPos(),
+							FPopupTransitionEffect::TypeInPopup);
+				    });
+
+				    FToolMenuEntry Entry = FToolMenuEntry::InitMenuEntry(
+					    FName("AddNewCurve"),
+					    LOCTEXT("AddNewCurve_Label", "Add New Curve"),
+					    LOCTEXT("AddNewCurve_Tooltip", "Add a new curve value"),
+					    FSlateIcon(),
+					    FToolUIActionChoice(Action),
+						EUserInterfaceActionType::Button);
+				
+				    InSubMenu->AddMenuEntry("AddNewCurve", Entry);
+
+					const FHierarchyTable_TableType_Skeleton& SkeletonTableMetadata = HierarchyTable->TableMetadata.Get<FHierarchyTable_TableType_Skeleton>();
+
+					FPersonaModule& PersonaModule = FModuleManager::LoadModuleChecked<FPersonaModule>("Persona");
+					InSubMenu->AddMenuEntry("ExistingCurveMenu", FToolMenuEntry::InitWidget(
+						"ExistingCurveMenu",
+						SNew(SVerticalBox)
+						+SVerticalBox::Slot()
+						.AutoHeight()
+						[
+							PersonaModule.CreateCurvePicker(SkeletonTableMetadata.Skeleton,
+								FOnCurvePicked::CreateLambda([&](const FName& InName)
+									{
+										const int32 SelectedIndex = Editor->GetSelectedEntryIndex();
+										const int32 ParentIndex = SelectedIndex == INDEX_NONE ? 0 : SelectedIndex;
+										AddCurve(HierarchyTable, ParentIndex, InName);
+										FSlateApplication::Get().DismissAllMenus();
+									}))
+						],
+						FText(),
+						true,
+						false,
+						true
+					));
+				}),
+			LOCTEXT("AddCurve_Label", "Add Curve"),
+			LOCTEXT("AddCurve_ToolTip", "Add a new curve to the hierarchy"),
+			FSlateIcon(FAppStyle::Get().GetStyleSetName(), "Icons.Plus")
+		));
+
+		HierarchyTableSection.AddEntry(FToolMenuEntry::InitComboButton(
+			"AddAttribute",
+			FUIAction(),
+			FNewToolMenuDelegate::CreateLambda([=, this](UToolMenu* InSubMenu)
+				{
+					FToolUIAction Action;
+				    Action.ExecuteAction = FToolMenuExecuteAction::CreateLambda([=, this](const FToolMenuContext& Context)
+				    {
+						TSharedRef<STextEntryPopup> TextEntry = SNew(STextEntryPopup)
+							.Label(LOCTEXT("NewAttributeEntryLabal", "Attribute Name"))
+							.OnTextCommitted_Lambda([=, this](const FText& CommittedText, ETextCommit::Type CommitInfo)
+								{
+									const int32 SelectedIndex = Editor->GetSelectedEntryIndex();
+									const int32 ParentIndex = SelectedIndex == INDEX_NONE ? 0 : SelectedIndex;
+									AddAttribute(HierarchyTable, ParentIndex, FName(CommittedText.ToString()));
+									FSlateApplication::Get().DismissAllMenus();
+								});
+
+						FSlateApplication& SlateApp = FSlateApplication::Get();
+						SlateApp.PushMenu(
+							SlateApp.GetInteractiveTopLevelWindows()[0],
+							FWidgetPath(),
+							TextEntry,
+							SlateApp.GetCursorPos(),
+							FPopupTransitionEffect::TypeInPopup);
+				    });
+
+				    FToolMenuEntry Entry = FToolMenuEntry::InitMenuEntry(
+					    FName("AddNewAttribute"),
+					    LOCTEXT("AddNewAttribute_Label", "Add New Attribute"),
+					    LOCTEXT("AddNewAttribute_Tooltip", "Add a new attribute value"),
+					    FSlateIcon(),
+					    FToolUIActionChoice(Action),
+						EUserInterfaceActionType::Button);
+				
+				    InSubMenu->AddMenuEntry("AddNewAttribute", Entry);
+				}),
+			LOCTEXT("AddAttribute_Label", "Add Attribute"),
+			LOCTEXT("AddAttribute_ToolTip", "Add a new attribute to the hierarchy"),
+			FSlateIcon(FAppStyle::Get().GetStyleSetName(), "Icons.Plus")
+		));
+	}
+}
+
+TSharedRef<SWidget> UHierarchyTable_TableTypeHandler_Skeleton::CreateLabelWidget(TObjectPtr<UHierarchyTable> HierarchyTable, const int32 EntryIndex) const
+{
+	const FSlateBrush* LabelIcon = nullptr;
+	FSlateColor LabelIconColor = FSlateColor::UseForeground();
+
+	const FHierarchyTableEntryData* const Entry = HierarchyTable->GetTableEntry(EntryIndex);
+	const FHierarchyTable_TablePayloadType_Skeleton& EntryTablePayload = Entry->TablePayload.Get<FHierarchyTable_TablePayloadType_Skeleton>();
+
+	switch (EntryTablePayload.EntryType)
+	{
+		case ESkeletonHierarchyTable_TablePayloadEntryType::Bone:
+			LabelIcon = FAppStyle::GetBrush("SkeletonTree.Bone");
+			break;
+		case ESkeletonHierarchyTable_TablePayloadEntryType::Curve:
+			LabelIcon = FAppStyle::GetBrush("AnimGraph.Attribute.Curves.Icon");
+			LabelIconColor = FAppStyle::GetSlateColor("AnimGraph.Attribute.Curves.Color");
+			break;
+		case ESkeletonHierarchyTable_TablePayloadEntryType::Attribute:
+			LabelIcon = FAppStyle::GetBrush("AnimGraph.Attribute.Attributes.Icon");
+			LabelIconColor = FAppStyle::GetSlateColor("AnimGraph.Attribute.Attributes.Color");
+			break;
+	}
+
+	return SNew(SHorizontalBox)
+		+SHorizontalBox::Slot()
+		.AutoWidth()
+		[
+			SNew(SImage)
+			.Image(LabelIcon)
+			.ColorAndOpacity(LabelIconColor)
+		]
+		+SHorizontalBox::Slot()
+		.AutoWidth()
+		[
+			SNew(SSpacer)
+			.Size(FVector2D(5.0f, 0.0f))
+		]
+		+SHorizontalBox::Slot()
+		.FillWidth(1.0f)
+		[
+			SNew(STextBlock)
+			.Text(FText::FromName(Entry->Identifier))
+		];
+}
+
+void UHierarchyTable_TableTypeHandler_Skeleton::AddCurve(TObjectPtr<UHierarchyTable> HierarchyTable, const int32 ParentIndex, const FName Identifier) const
+{
+	FHierarchyTableEntryData NewEntry;
+	NewEntry.Parent = ParentIndex;
+	NewEntry.Identifier = Identifier;
+	NewEntry.Payload = TOptional<FInstancedStruct>();
+	NewEntry.OwnerTable = HierarchyTable;
+	{
+		FInstancedStruct TablePayload;
+		TablePayload.InitializeAs<FHierarchyTable_TablePayloadType_Skeleton>();
+		FHierarchyTable_TablePayloadType_Skeleton& SkeletonTablePayload = TablePayload.GetMutable<FHierarchyTable_TablePayloadType_Skeleton>();
+		SkeletonTablePayload.EntryType = ESkeletonHierarchyTable_TablePayloadEntryType::Curve;
+		
+		NewEntry.TablePayload = TablePayload;
+	}
+
+	HierarchyTable->AddEntry(NewEntry);
+}
+	
+void UHierarchyTable_TableTypeHandler_Skeleton::AddAttribute(TObjectPtr<UHierarchyTable> HierarchyTable, const int32 ParentIndex, const FName Identifier) const
+{
+	FHierarchyTableEntryData NewEntry;
+	NewEntry.Parent = ParentIndex;
+	NewEntry.Identifier = Identifier;
+	NewEntry.Payload = TOptional<FInstancedStruct>();
+	NewEntry.OwnerTable = HierarchyTable;
+	{
+		FInstancedStruct TablePayload;
+		TablePayload.InitializeAs<FHierarchyTable_TablePayloadType_Skeleton>();
+		FHierarchyTable_TablePayloadType_Skeleton& SkeletonTablePayload = TablePayload.GetMutable<FHierarchyTable_TablePayloadType_Skeleton>();
+		SkeletonTablePayload.EntryType = ESkeletonHierarchyTable_TablePayloadEntryType::Attribute;
+		
+		NewEntry.TablePayload = TablePayload;
+	}
+
+	HierarchyTable->AddEntry(NewEntry);
+}
+
+#undef LOCTEXT_NAMESPACE
\ No newline at end of file
diff --git a/Engine/Plugins/Experimental/Animation/HierarchyTableBuiltin/Source/Editor/Private/SkeletonHierarchyTableTypeHandler.h b/Engine/Plugins/Experimental/Animation/HierarchyTableBuiltin/Source/Editor/Private/SkeletonHierarchyTableTypeHandler.h
new file mode 100644
--- /dev/null
+++ b/Engine/Plugins/Experimental/Animation/HierarchyTableBuiltin/Source/Editor/Private/SkeletonHierarchyTableTypeHandler.h
@@ -0,0 +1,29 @@
+// Copyright Epic Games, Inc. All Rights Reserved.
+
+#pragma once
+
+#include "HierarchyTableTypeHandler.h"
+
+#include "SkeletonHierarchyTableTypeHandler.generated.h"
+
+UCLASS()
+class UHierarchyTable_TableTypeHandler_Skeleton final : public UHierarchyTable_TableTypeHandler
+{
+	GENERATED_BODY()
+
+public:
+	virtual TSharedPtr<FHierarchyTable_TableType> CreateTableType() const;
+
+	virtual void InitialiseTable(const FInstancedStruct& TableType, TObjectPtr<UHierarchyTable> HierarchyTable) const;
+
+	virtual bool FactoryConfigureProperties(FInstancedStruct& TableType) const override;
+
+	virtual void ExtendToolbar(UToolMenu* ToolMenu, UHierarchyTable* HierarchyTable, const IHierarchyTableEditor* Editor) const override;
+
+	virtual TSharedRef<SWidget> CreateLabelWidget(TObjectPtr<UHierarchyTable> HierarchyTable, const int32 EntryIndex) const;
+
+private:
+	void AddCurve(TObjectPtr<UHierarchyTable> HierarchyTable, const int32 ParentIndex, const FName Identifier) const;
+	
+	void AddAttribute(TObjectPtr<UHierarchyTable> HierarchyTable, const int32 ParentIndex, const FName Identifier) const;
+};
\ No newline at end of file
diff --git a/Engine/Plugins/Experimental/Animation/HierarchyTableBuiltin/Source/Editor/Private/TimeProfile/TimeProfileTypeHandler.cpp b/Engine/Plugins/Experimental/Animation/HierarchyTableBuiltin/Source/Editor/Private/TimeProfile/TimeProfileTypeHandler.cpp
--- a/Engine/Plugins/Experimental/Animation/HierarchyTableBuiltin/Source/Editor/Private/TimeProfile/TimeProfileTypeHandler.cpp
+++ b/Engine/Plugins/Experimental/Animation/HierarchyTableBuiltin/Source/Editor/Private/TimeProfile/TimeProfileTypeHandler.cpp
@@ -4,11 +4,11 @@
 #include "TimeProfile/HierarchyTableTypeTime.h"
 #include "TimeProfile/TimeProfileProxyColumn.h"
 
-FInstancedStruct UHierarchyTableTypeHandler_Time::GetDefaultEntry() const
+FInstancedStruct UHierarchyTable_ElementTypeHandler_Time::GetDefaultEntry() const
 {
 	FInstancedStruct NewEntry;
-	NewEntry.InitializeAs(FHierarchyTableType_Time::StaticStruct());
-	FHierarchyTableType_Time& NewEntryRef = NewEntry.GetMutable<FHierarchyTableType_Time>();
+	NewEntry.InitializeAs(FHierarchyTable_ElementType_Time::StaticStruct());
+	FHierarchyTable_ElementType_Time& NewEntryRef = NewEntry.GetMutable<FHierarchyTable_ElementType_Time>();
 	NewEntryRef.StartTime = 0.0f;
 	NewEntryRef.EndTime = 1.0f;
 	NewEntryRef.TimeFactor = 1.0f;
@@ -16,7 +16,7 @@ FInstancedStruct UHierarchyTableTypeHandler_Time::GetDefaultEntry() const
 	return NewEntry;
 }
 
-TArray<UScriptStruct*> UHierarchyTableTypeHandler_Time::GetColumns() const
+TArray<UScriptStruct*> UHierarchyTable_ElementTypeHandler_Time::GetColumns() const
 {
 	return { FHierarchyTableTimeColumn_StartTime::StaticStruct(), FHierarchyTableTimeColumn_EndTime::StaticStruct(), FHierarchyTableTimeColumn_TimeFactor::StaticStruct() };
 }
\ No newline at end of file
diff --git a/Engine/Plugins/Experimental/Animation/HierarchyTableBuiltin/Source/Editor/Private/TimeProfile/TimeProfileTypeHandler.h b/Engine/Plugins/Experimental/Animation/HierarchyTableBuiltin/Source/Editor/Private/TimeProfile/TimeProfileTypeHandler.h
--- a/Engine/Plugins/Experimental/Animation/HierarchyTableBuiltin/Source/Editor/Private/TimeProfile/TimeProfileTypeHandler.h
+++ b/Engine/Plugins/Experimental/Animation/HierarchyTableBuiltin/Source/Editor/Private/TimeProfile/TimeProfileTypeHandler.h
@@ -7,7 +7,7 @@
 #include "TimeProfileTypeHandler.generated.h"
 
 UCLASS()
-class UHierarchyTableTypeHandler_Time final : public UHierarchyTableTypeHandler_Base
+class UHierarchyTable_ElementTypeHandler_Time final : public UHierarchyTable_ElementTypeHandler
 {
 	GENERATED_BODY()
 
diff --git a/Engine/Plugins/Experimental/Animation/HierarchyTableBuiltin/Source/Editor/Private/TimeProfile/TimeProfileWidgetConstructor.cpp b/Engine/Plugins/Experimental/Animation/HierarchyTableBuiltin/Source/Editor/Private/TimeProfile/TimeProfileWidgetConstructor.cpp
--- a/Engine/Plugins/Experimental/Animation/HierarchyTableBuiltin/Source/Editor/Private/TimeProfile/TimeProfileWidgetConstructor.cpp
+++ b/Engine/Plugins/Experimental/Animation/HierarchyTableBuiltin/Source/Editor/Private/TimeProfile/TimeProfileWidgetConstructor.cpp
@@ -8,51 +8,51 @@
 TSharedRef<SWidget> FHierarchyTableTimeWidgetConstructor_StartTime::CreateInternalWidget(UHierarchyTable* HierarchyTable, int32 EntryIndex)
 {
 	return SNew(SSpinBox<float>)
-		.IsEnabled_Lambda([HierarchyTable, EntryIndex]() { return HierarchyTable->TableData[EntryIndex].IsOverridden(); })
+		.IsEnabled_Lambda([HierarchyTable, EntryIndex]() { return HierarchyTable->GetTableEntry(EntryIndex)->IsOverridden(); })
 		.MinDesiredWidth(100.0f)
 		.MinValue(0.0f)
 		.MaxValue(1.0f)
 		.Value_Lambda([HierarchyTable, EntryIndex]()
 			{
-				return HierarchyTable->TableData[EntryIndex].GetValue<FHierarchyTableType_Time>()->StartTime;
+				return HierarchyTable->GetTableEntry(EntryIndex)->GetValue<FHierarchyTable_ElementType_Time>()->StartTime;
 			})
 		.OnValueChanged_Lambda([HierarchyTable, EntryIndex](float NewValue)
 			{
-				HierarchyTable->TableData[EntryIndex].GetMutableValue<FHierarchyTableType_Time>()->StartTime = NewValue;
+				HierarchyTable->GetMutableTableEntry(EntryIndex)->GetMutableValue<FHierarchyTable_ElementType_Time>()->StartTime = NewValue;
 			});
 }
 
 TSharedRef<SWidget> FHierarchyTableTimeWidgetConstructor_EndTime::CreateInternalWidget(UHierarchyTable* HierarchyTable, int32 EntryIndex)
 {
 	return SNew(SSpinBox<float>)
-		.IsEnabled_Lambda([HierarchyTable, EntryIndex]() { return HierarchyTable->TableData[EntryIndex].IsOverridden(); })
+		.IsEnabled_Lambda([HierarchyTable, EntryIndex]() { return HierarchyTable->GetTableEntry(EntryIndex)->IsOverridden(); })
 		.MinDesiredWidth(100.0f)
 		.MinValue(0.0f)
 		.MaxValue(1.0f)
 		.Value_Lambda([HierarchyTable, EntryIndex]()
 			{
-				return HierarchyTable->TableData[EntryIndex].GetValue<FHierarchyTableType_Time>()->EndTime;
+				return HierarchyTable->GetTableEntry(EntryIndex)->GetValue<FHierarchyTable_ElementType_Time>()->EndTime;
 			})
 		.OnValueChanged_Lambda([HierarchyTable, EntryIndex](float NewValue)
 			{
-				HierarchyTable->TableData[EntryIndex].GetMutableValue<FHierarchyTableType_Time>()->EndTime = NewValue;
+				HierarchyTable->GetMutableTableEntry(EntryIndex)->GetMutableValue<FHierarchyTable_ElementType_Time>()->EndTime = NewValue;
 			});
 }
 
 TSharedRef<SWidget> FHierarchyTableTimeWidgetConstructor_TimeFactor::CreateInternalWidget(UHierarchyTable* HierarchyTable, int32 EntryIndex)
 {
 	return SNew(SSpinBox<float>)
-		.IsEnabled_Lambda([HierarchyTable, EntryIndex]() { return HierarchyTable->TableData[EntryIndex].IsOverridden(); })
+		.IsEnabled_Lambda([HierarchyTable, EntryIndex]() { return HierarchyTable->GetTableEntry(EntryIndex)->IsOverridden(); })
 		.MinDesiredWidth(100.0f)
 		.MinValue(0.0f)
 		.MaxValue(1.0f)
 		.Value_Lambda([HierarchyTable, EntryIndex]()
 			{
-				return HierarchyTable->TableData[EntryIndex].GetValue<FHierarchyTableType_Time>()->TimeFactor;
+				return HierarchyTable->GetTableEntry(EntryIndex)->GetValue<FHierarchyTable_ElementType_Time>()->TimeFactor;
 			})
 		.OnValueChanged_Lambda([HierarchyTable, EntryIndex](float NewValue)
 			{
-				HierarchyTable->TableData[EntryIndex].GetMutableValue<FHierarchyTableType_Time>()->TimeFactor = NewValue;
+				HierarchyTable->GetMutableTableEntry(EntryIndex)->GetMutableValue<FHierarchyTable_ElementType_Time>()->TimeFactor = NewValue;
 			});
 }
 
diff --git a/Engine/Plugins/Experimental/Animation/HierarchyTableBuiltin/Source/Runtime/Private/AnimNode_BlendProfileLayeredBlend.cpp b/Engine/Plugins/Experimental/Animation/HierarchyTableBuiltin/Source/Runtime/Private/AnimNode_BlendProfileLayeredBlend.cpp
--- a/Engine/Plugins/Experimental/Animation/HierarchyTableBuiltin/Source/Runtime/Private/AnimNode_BlendProfileLayeredBlend.cpp
+++ b/Engine/Plugins/Experimental/Animation/HierarchyTableBuiltin/Source/Runtime/Private/AnimNode_BlendProfileLayeredBlend.cpp
@@ -11,6 +11,7 @@
 #include "Animation/AttributeTypes.h"
 #include "Animation/IAttributeBlendOperator.h"
 #include "Hash/CityHash.h"
+#include "SkeletonHierarchyTableType.h"
 
 #include UE_INLINE_GENERATED_CPP_BY_NAME(AnimNode_BlendProfileLayeredBlend)
 
@@ -52,11 +53,11 @@ void FAnimNode_BlendProfileLayeredBlend::UpdateCachedBoneData(const FBoneContain
 		CachedCurveMaskWeights.Empty();
 		//CachedAttributeMaskWeights.Empty();
 
-		for (const FHierarchyTableEntryData& Entry : BlendMask->TableData)
+		for (const FHierarchyTableEntryData& Entry : MaskTable->GetTableData())
 		{
-			if (Entry.EntryType == EHierarchyTableEntryType::Curve)
+			if (Entry.TablePayload.Get<FHierarchyTable_TablePayloadType_Skeleton>().EntryType == ESkeletonHierarchyTable_TablePayloadEntryType::Curve)
 			{
-				const float EntryWeight = Entry.GetValue<FHierarchyTableType_Mask>()->Value;
+				const float EntryWeight = Entry.GetValue<FHierarchyTable_ElementType_Mask>()->Value;
 				CachedCurveMaskWeights.Add<UE::Anim::FCurveElement>({ Entry.Identifier, EntryWeight });
 			}
 			/*
@@ -307,17 +308,20 @@ void FAnimNode_BlendProfileLayeredBlend::CreateMaskWeights(const USkeleton* Skel
 	PerBoneBlendWeights.Reset(NumBones);
 	PerBoneBlendWeights.AddZeroed(NumBones);
 
-	check(BlendMask->Skeleton == Skeleton && BlendMask->TableType == FHierarchyTableType_Mask::StaticStruct());
+	check(MaskTable->TableMetadata.GetScriptStruct() == FHierarchyTable_TableType_Skeleton::StaticStruct());
+	check(MaskTable->ElementType == FHierarchyTable_ElementType_Mask::StaticStruct());
+
+	const FHierarchyTable_TableType_Skeleton& SkeletonMetadata = MaskTable->TableMetadata.Get<FHierarchyTable_TableType_Skeleton>();
 
-	for (int32 EntryIndex = 0; EntryIndex < BlendMask->TableData.Num(); ++EntryIndex)
+	for (int32 EntryIndex = 0; EntryIndex < MaskTable->GetTableData().Num(); ++EntryIndex)
 	{
-		const FHierarchyTableEntryData& EntryData = BlendMask->TableData[EntryIndex];
-		if (EntryData.EntryType != EHierarchyTableEntryType::Bone)
+		const FHierarchyTableEntryData* const EntryData = MaskTable->GetTableEntry(EntryIndex);
+		if (EntryData->TablePayload.Get<FHierarchyTable_TablePayloadType_Skeleton>().EntryType != ESkeletonHierarchyTable_TablePayloadEntryType::Bone)
 		{
 			break;
 		}
 
-		const FHierarchyTableType_Mask* MaskValue = EntryData.GetValue<FHierarchyTableType_Mask>();
+		const FHierarchyTable_ElementType_Mask* MaskValue = EntryData->GetValue<FHierarchyTable_ElementType_Mask>();
 		check(MaskValue);
 			
 		PerBoneBlendWeights[EntryIndex] = MaskValue->Value;
diff --git a/Engine/Plugins/Experimental/Animation/HierarchyTableBuiltin/Source/Runtime/Public/AnimNode_BlendProfileLayeredBlend.h b/Engine/Plugins/Experimental/Animation/HierarchyTableBuiltin/Source/Runtime/Public/AnimNode_BlendProfileLayeredBlend.h
--- a/Engine/Plugins/Experimental/Animation/HierarchyTableBuiltin/Source/Runtime/Public/AnimNode_BlendProfileLayeredBlend.h
+++ b/Engine/Plugins/Experimental/Animation/HierarchyTableBuiltin/Source/Runtime/Public/AnimNode_BlendProfileLayeredBlend.h
@@ -26,10 +26,10 @@ public:
 	FPoseLink BlendPose;
 
 	/**
-	 * The blend mask to use to control layering of the pose, curves, and attributes
+	 * The mask hierarchy table to use to control layering of the pose, curves, and attributes
 	 */
 	UPROPERTY(EditAnywhere, BlueprintReadWrite, Category = Config)
-	TObjectPtr<UHierarchyTable> BlendMask;
+	TObjectPtr<UHierarchyTable> MaskTable;
 
 protected:
 	// Per-bone weights for the skeleton. Serialized as these are only relative to the skeleton, but can potentially
diff --git a/Engine/Plugins/Experimental/Animation/HierarchyTableBuiltin/Source/Runtime/Public/MaskProfile/HierarchyTableTypeMask.h b/Engine/Plugins/Experimental/Animation/HierarchyTableBuiltin/Source/Runtime/Public/MaskProfile/HierarchyTableTypeMask.h
--- a/Engine/Plugins/Experimental/Animation/HierarchyTableBuiltin/Source/Runtime/Public/MaskProfile/HierarchyTableTypeMask.h
+++ b/Engine/Plugins/Experimental/Animation/HierarchyTableBuiltin/Source/Runtime/Public/MaskProfile/HierarchyTableTypeMask.h
@@ -7,11 +7,11 @@
 #include "HierarchyTableTypeMask.generated.h"
 
 USTRUCT()
-struct FHierarchyTableType_Mask final : public FHierarchyTableType
+struct FHierarchyTable_ElementType_Mask final : public FHierarchyTable_ElementType
 {
 	GENERATED_BODY()
 
-	FHierarchyTableType_Mask()
+	FHierarchyTable_ElementType_Mask()
 		: Value(0.0f)
 	{
 	}
diff --git a/Engine/Plugins/Experimental/Animation/HierarchyTableBuiltin/Source/Runtime/Public/SkeletonHierarchyTableType.h b/Engine/Plugins/Experimental/Animation/HierarchyTableBuiltin/Source/Runtime/Public/SkeletonHierarchyTableType.h
new file mode 100644
--- /dev/null
+++ b/Engine/Plugins/Experimental/Animation/HierarchyTableBuiltin/Source/Runtime/Public/SkeletonHierarchyTableType.h
@@ -0,0 +1,35 @@
+// Copyright Epic Games, Inc. All Rights Reserved.
+
+#pragma once
+
+#include "HierarchyTableType.h"
+
+#include "SkeletonHierarchyTableType.generated.h"
+
+class USkeleton;
+
+USTRUCT()
+struct FHierarchyTable_TableType_Skeleton final : public FHierarchyTable_TableType
+{
+	GENERATED_BODY()
+
+	UPROPERTY()
+	TObjectPtr<USkeleton> Skeleton;
+};
+
+UENUM()
+enum class ESkeletonHierarchyTable_TablePayloadEntryType
+{
+	Bone,
+	Curve,
+	Attribute
+};
+
+USTRUCT()
+struct FHierarchyTable_TablePayloadType_Skeleton final : public FHierarchyTable_TablePayloadType
+{
+	GENERATED_BODY()
+
+	UPROPERTY()
+	ESkeletonHierarchyTable_TablePayloadEntryType EntryType;
+};
\ No newline at end of file
diff --git a/Engine/Plugins/Experimental/Animation/HierarchyTableBuiltin/Source/Runtime/Public/TimeProfile/HierarchyTableTypeTime.h b/Engine/Plugins/Experimental/Animation/HierarchyTableBuiltin/Source/Runtime/Public/TimeProfile/HierarchyTableTypeTime.h
--- a/Engine/Plugins/Experimental/Animation/HierarchyTableBuiltin/Source/Runtime/Public/TimeProfile/HierarchyTableTypeTime.h
+++ b/Engine/Plugins/Experimental/Animation/HierarchyTableBuiltin/Source/Runtime/Public/TimeProfile/HierarchyTableTypeTime.h
@@ -7,11 +7,11 @@
 #include "HierarchyTableTypeTime.generated.h"
 
 USTRUCT()
-struct FHierarchyTableType_Time final : public FHierarchyTableType
+struct FHierarchyTable_ElementType_Time final : public FHierarchyTable_ElementType
 {
 	GENERATED_BODY()
 
-	FHierarchyTableType_Time()
+	FHierarchyTable_ElementType_Time()
 		: StartTime(0.0f)
 		, EndTime(0.0f)
 		, TimeFactor(0.0f)
diff --git a/Engine/Plugins/Experimental/Animation/HierarchyTableBuiltin/Source/UncookedOnly/Private/AnimGraphNode_BlendProfileLayeredBlend.cpp b/Engine/Plugins/Experimental/Animation/HierarchyTableBuiltin/Source/UncookedOnly/Private/AnimGraphNode_BlendProfileLayeredBlend.cpp
--- a/Engine/Plugins/Experimental/Animation/HierarchyTableBuiltin/Source/UncookedOnly/Private/AnimGraphNode_BlendProfileLayeredBlend.cpp
+++ b/Engine/Plugins/Experimental/Animation/HierarchyTableBuiltin/Source/UncookedOnly/Private/AnimGraphNode_BlendProfileLayeredBlend.cpp
@@ -42,10 +42,10 @@ FString UAnimGraphNode_BlendProfileLayeredBlend::GetNodeCategory() const
 
 void UAnimGraphNode_BlendProfileLayeredBlend::PreloadRequiredAssets()
 {
-	TObjectPtr<UHierarchyTable> BlendMask = Node.BlendMask;
-	if (BlendMask)
+	TObjectPtr<UHierarchyTable> MaskTable = Node.MaskTable;
+	if (MaskTable)
 	{
-		PreloadObject(BlendMask);
+		PreloadObject(MaskTable);
 	}
 
 	Super::PreloadRequiredAssets();
@@ -57,16 +57,16 @@ void UAnimGraphNode_BlendProfileLayeredBlend::ValidateAnimNodeDuringCompilation(
 
 	bool bCompilationError = false;
 
-	const TObjectPtr<UHierarchyTable> BlendMask = Node.BlendMask;
+	const TObjectPtr<UHierarchyTable> MaskTable = Node.MaskTable;
 
-	if (BlendMask == nullptr)
+	if (MaskTable == nullptr)
 	{
-		MessageLog.Error(*LOCTEXT("InvalidMask", "@@ has a null blend profile.").ToString(), this, BlendMask);
+		MessageLog.Error(*LOCTEXT("InvalidMask", "@@ has a null mask table.").ToString(), this, MaskTable);
 		bCompilationError = true;
 	}
-	else if (BlendMask->TableType != FHierarchyTableType_Mask::StaticStruct())
+	else if (MaskTable->ElementType != FHierarchyTable_ElementType_Mask::StaticStruct())
 	{
-		MessageLog.Error(*LOCTEXT("InvalidMaskType", "@@ has a null blend profile that is not a mask type.").ToString(), this, BlendMask);
+		MessageLog.Error(*LOCTEXT("InvalidMaskType", "@@ has a null mask table that is not a mask type.").ToString(), this, MaskTable);
 		bCompilationError = true;
 	}
 
