Hierarchy Table Cleanup

diff --git a/Engine/Plugins/Experimental/Animation/HierarchyTable/Source/Editor/Private/DefaultHierarchyTableTypeHandler.cpp b/Engine/Plugins/Experimental/Animation/HierarchyTable/Source/Editor/Private/DefaultHierarchyTableTypeHandler.cpp
--- a/Engine/Plugins/Experimental/Animation/HierarchyTable/Source/Editor/Private/DefaultHierarchyTableTypeHandler.cpp
+++ b/Engine/Plugins/Experimental/Animation/HierarchyTable/Source/Editor/Private/DefaultHierarchyTableTypeHandler.cpp
@@ -8,7 +8,7 @@
 
 #define LOCTEXT_NAMESPACE "UHierarchyTable_TableTypeHandler_Default"
 
-void UHierarchyTable_TableTypeHandler_Default::InitialiseTable(const FInstancedStruct& TableType, TObjectPtr<UHierarchyTable> HierarchyTable) const
+void UHierarchyTable_TableTypeHandler_Default::ConstructHierarchy()
 {
 	// Do nothing, hierarchy is built manually
 }
diff --git a/Engine/Plugins/Experimental/Animation/HierarchyTable/Source/Editor/Private/DefaultHierarchyTableTypeHandler.h b/Engine/Plugins/Experimental/Animation/HierarchyTable/Source/Editor/Private/DefaultHierarchyTableTypeHandler.h
--- a/Engine/Plugins/Experimental/Animation/HierarchyTable/Source/Editor/Private/DefaultHierarchyTableTypeHandler.h
+++ b/Engine/Plugins/Experimental/Animation/HierarchyTable/Source/Editor/Private/DefaultHierarchyTableTypeHandler.h
@@ -12,7 +12,7 @@ class UHierarchyTable_TableTypeHandler_Default final : public UHierarchyTable_Ta
 	GENERATED_BODY()
 
 public:
-	virtual void InitialiseTable(const FInstancedStruct& TableType, TObjectPtr<UHierarchyTable> HierarchyTable) const;
+	virtual void ConstructHierarchy() override;
 
 	virtual bool FactoryConfigureProperties(FInstancedStruct& TableType) const override;
 };
\ No newline at end of file
diff --git a/Engine/Plugins/Experimental/Animation/HierarchyTable/Source/Editor/Private/HierachyTableEditorModule.cpp b/Engine/Plugins/Experimental/Animation/HierarchyTable/Source/Editor/Private/HierachyTableEditorModule.cpp
--- a/Engine/Plugins/Experimental/Animation/HierarchyTable/Source/Editor/Private/HierachyTableEditorModule.cpp
+++ b/Engine/Plugins/Experimental/Animation/HierarchyTable/Source/Editor/Private/HierachyTableEditorModule.cpp
@@ -7,7 +7,7 @@
 
 void FHierarchyTableEditorModule::StartupModule()
 {
-	RegisterTableType(FHierarchyTable_TableType_Default::StaticStruct(), GetDefault<UHierarchyTable_TableTypeHandler_Default>());
+	RegisterTableType(FHierarchyTable_TableType_Default::StaticStruct(), UHierarchyTable_TableTypeHandler_Default::StaticClass());
 }
 
 void FHierarchyTableEditorModule::ShutdownModule()
@@ -15,7 +15,7 @@ void FHierarchyTableEditorModule::ShutdownModule()
 	UnregisterTableType(FHierarchyTable_TableType_Default::StaticStruct());
 }
 
-void FHierarchyTableEditorModule::RegisterTableType(const UScriptStruct* TableType, const UHierarchyTable_TableTypeHandler* Handler)
+void FHierarchyTableEditorModule::RegisterTableType(const UScriptStruct* TableType, const UClass* Handler)
 {
 	TableHandlers.Add(TableType, Handler);
 }
@@ -25,34 +25,48 @@ void FHierarchyTableEditorModule::UnregisterTableType(const UScriptStruct* Table
 	TableHandlers.Remove(TableType);
 }
 
-const UHierarchyTable_TableTypeHandler* FHierarchyTableEditorModule::FindTableHandler(const UScriptStruct* TableType) const
+const TObjectPtr<UHierarchyTable_TableTypeHandler> FHierarchyTableEditorModule::CreateTableHandler(const TObjectPtr<UHierarchyTable> HierarchyTable)
 {
-	const UHierarchyTable_TableTypeHandler* const* Result = TableHandlers.Find(TableType);
+	const UClass** Result = TableHandlers.Find(HierarchyTable->GetTableMetadataStruct());
 	if (Result)
 	{
-		return *Result;
+		TObjectPtr<UHierarchyTable_TableTypeHandler> Handler = NewObject<UHierarchyTable_TableTypeHandler>((UObject*)GetTransientPackage(), *Result);
+		Handler->SetHierarchyTable(HierarchyTable);
+
+		return Handler;
 	}
 	return nullptr;
 }
 
-void FHierarchyTableEditorModule::RegisterElementType(const UScriptStruct* ElementType, const UHierarchyTable_ElementTypeHandler* Handler)
+const TObjectPtr<UHierarchyTable_TableTypeHandler> FHierarchyTableEditorModule::CreateTableHandler(const TObjectPtr<const UScriptStruct> TableType)
 {
-	ElementHandlers.Add(ElementType, Handler);
+	const UClass** Result = TableHandlers.Find(TableType);
+	if (Result)
+	{
+		return NewObject<UHierarchyTable_TableTypeHandler>((UObject*)GetTransientPackage(), *Result);
+	}
+	return nullptr;
 }
 
-void FHierarchyTableEditorModule::UnregisterElementType(const UScriptStruct* ElementType)
+void FHierarchyTableEditorModule::RegisterElementTypeEditorColumns(const UScriptStruct* ElementType, const TArray<TObjectPtr<UScriptStruct>> Columns)
 {
-	ElementHandlers.Remove(ElementType);
+	EditorColumns.Add(ElementType, Columns);
 }
 
-const UHierarchyTable_ElementTypeHandler* FHierarchyTableEditorModule::FindElementHandler(const UScriptStruct* ElementType) const
+void FHierarchyTableEditorModule::UnregisterElementTypeEditorColumns(const UScriptStruct* ElementType)
 {
-	const UHierarchyTable_ElementTypeHandler* const* Result = ElementHandlers.Find(ElementType);
+	EditorColumns.Remove(ElementType);
+}
+
+TArray<UScriptStruct*> FHierarchyTableEditorModule::GetElementTypeEditorColumns(const TObjectPtr<const UHierarchyTable> HierarchyTable)
+{
+	const TArray<UScriptStruct*>* Result = EditorColumns.Find(HierarchyTable->GetElementType());
+
 	if (Result)
 	{
 		return *Result;
 	}
-	return nullptr;
+	return TArray<UScriptStruct*>();
 }
 
 IMPLEMENT_MODULE(FHierarchyTableEditorModule, HierarchyTableEditor)
diff --git a/Engine/Plugins/Experimental/Animation/HierarchyTable/Source/Editor/Private/HierarchyTableEditorToolkit.cpp b/Engine/Plugins/Experimental/Animation/HierarchyTable/Source/Editor/Private/HierarchyTableEditorToolkit.cpp
--- a/Engine/Plugins/Experimental/Animation/HierarchyTable/Source/Editor/Private/HierarchyTableEditorToolkit.cpp
+++ b/Engine/Plugins/Experimental/Animation/HierarchyTable/Source/Editor/Private/HierarchyTableEditorToolkit.cpp
@@ -118,19 +118,14 @@ TSharedRef<SWidget> FHierarchyTableEditorToolkit::CreateTedsOutliner()
 			.Text(INVTEXT("You need to enable the Typed Element Data Storage plugin to see the table viewer!"));
 	}
 
-	if (!ensure(HierarchyTable->ElementType))
+	if (!ensure(HierarchyTable->GetElementType()))
 	{
 		return SNullWidget::NullWidget;
 	}
 
 	FHierarchyTableEditorModule& HierarchyTableModule = FModuleManager::GetModuleChecked<FHierarchyTableEditorModule>("HierarchyTableEditor");
-	const UHierarchyTable_ElementTypeHandler* Handler = HierarchyTableModule.FindElementHandler(HierarchyTable->ElementType);
-	if (!ensureMsgf(Handler, TEXT("Could not find handler for %s, have you forgotten to register it?"), *HierarchyTable->ElementType->GetName()))
-	{
-		return SNullWidget::NullWidget;
-	}
 
-	TArray<UScriptStruct*> HierarchyTableTypeColumns = Handler->GetColumns();
+	TArray<UScriptStruct*> HierarchyTableTypeColumns = HierarchyTableModule.GetElementTypeEditorColumns(HierarchyTable);
 	HierarchyTableTypeColumns.Add(FTypedElementOverrideColumn::StaticStruct());
 
 	FQueryDescription ColumnQueryDescription =
@@ -164,7 +159,7 @@ TSharedRef<SWidget> FHierarchyTableEditorToolkit::CreateTedsOutliner()
 	IEditorDataStorageProvider* DSI = GetMutableDataStorageFeature<IEditorDataStorageProvider>(StorageFeatureName);
 	static TableHandle Table = DSI->FindTable(FName("Editor_HierarchyTableTable"));
 	
-	TArray<UScriptStruct*> BaseHierarchyTableTypeColumns = Handler->GetColumns();
+	TArray<UScriptStruct*> BaseHierarchyTableTypeColumns = HierarchyTableModule.GetElementTypeEditorColumns(HierarchyTable);
 
 	for (int32 EntryIndex = 0; EntryIndex < HierarchyTable->GetTableData().Num(); ++EntryIndex)
 	{
@@ -204,7 +199,7 @@ TSharedRef<SWidget> FHierarchyTableEditorToolkit::CreateTedsOutliner()
 void FHierarchyTableEditorToolkit::ExtendToolbar()
 {
 	FHierarchyTableEditorModule& HierarchyTableModule = FModuleManager::GetModuleChecked<FHierarchyTableEditorModule>("HierarchyTableEditor");
-	const UHierarchyTable_TableTypeHandler* Handler = HierarchyTableModule.FindTableHandler(HierarchyTable->TableMetadata.GetScriptStruct());
+	const TObjectPtr<UHierarchyTable_TableTypeHandler> Handler = HierarchyTableModule.CreateTableHandler(HierarchyTable);
 	if (!ensure(Handler))
 	{
 		return;
@@ -216,7 +211,7 @@ void FHierarchyTableEditorToolkit::ExtendToolbar()
 	static const FName MenuName = GetToolMenuToolbarName(ParentName);
 
 	UToolMenu* ToolMenu = UToolMenus::Get()->ExtendMenu(MenuName);
-	Handler->ExtendToolbar(ToolMenu, HierarchyTable, this);
+	Handler->ExtendToolbar(ToolMenu, this);
 }
 
 int32 FHierarchyTableEditorToolkit::GetSelectedEntryIndex() const
@@ -245,13 +240,6 @@ void FHierarchyTableEditorToolkit::HandleOnEntryAdded(const int32 EntryIndex)
 	IEditorDataStorageProvider* DSI = GetMutableDataStorageFeature<IEditorDataStorageProvider>(StorageFeatureName);
 	static TableHandle Table = DSI->FindTable(FName("Editor_HierarchyTableTable"));
 
-	FHierarchyTableEditorModule& HierarchyTableModule = FModuleManager::GetModuleChecked<FHierarchyTableEditorModule>("HierarchyTableEditor");
-	const UHierarchyTable_ElementTypeHandler* ElementHandler = HierarchyTableModule.FindElementHandler(HierarchyTable->ElementType);
-	if (!ensure(ElementHandler))
-	{
-		return;
-	}
-	
 	const RowHandle Row = DSI->AddRow(Table);
 	{
 		FTypedElementOverrideColumn OverrideEntry;
@@ -270,7 +258,8 @@ void FHierarchyTableEditorToolkit::HandleOnEntryAdded(const int32 EntryIndex)
 		DSI->AddColumn<FTableRowParentColumn>(Row, { .Parent = *ParentRow });
 	}
 
-	for (const UScriptStruct* Column : ElementHandler->GetColumns())
+	FHierarchyTableEditorModule& HierarchyTableModule = FModuleManager::GetModuleChecked<FHierarchyTableEditorModule>("HierarchyTableEditor");
+	for (const UScriptStruct* Column : HierarchyTableModule.GetElementTypeEditorColumns(HierarchyTable))
 	{
 		DSI->AddColumn(Row, Column);
 	}
diff --git a/Engine/Plugins/Experimental/Animation/HierarchyTable/Source/Editor/Private/HierarchyTableFactory.cpp b/Engine/Plugins/Experimental/Animation/HierarchyTable/Source/Editor/Private/HierarchyTableFactory.cpp
--- a/Engine/Plugins/Experimental/Animation/HierarchyTable/Source/Editor/Private/HierarchyTableFactory.cpp
+++ b/Engine/Plugins/Experimental/Animation/HierarchyTable/Source/Editor/Private/HierarchyTableFactory.cpp
@@ -22,15 +22,14 @@ UHierarchyTableFactory::UHierarchyTableFactory()
 UObject* UHierarchyTableFactory::FactoryCreateNew(UClass* Class, UObject* InParent, FName Name, EObjectFlags Flags, UObject* Context, FFeedbackContext* Warn)
 {
 	TObjectPtr<UHierarchyTable> HierarchyTable = NewObject<UHierarchyTable>(InParent, Class, Name, Flags, Context);
-
-	HierarchyTable->TableMetadata = TableMetadata;
-	HierarchyTable->ElementType = ElementType;
+	HierarchyTable->Initialize(TableMetadata, ElementType);
 
 	FHierarchyTableEditorModule& HierarchyTableModule = FModuleManager::GetModuleChecked<FHierarchyTableEditorModule>("HierarchyTableEditor");
-	const UHierarchyTable_TableTypeHandler* TableHandler = HierarchyTableModule.FindTableHandler(TableMetadata.GetScriptStruct());
+	const TObjectPtr<UHierarchyTable_TableTypeHandler> TableHandler = HierarchyTableModule.CreateTableHandler(HierarchyTable);
 	check(TableHandler);
 
-	TableHandler->InitialiseTable(TableMetadata, HierarchyTable);
+	TableHandler->SetHierarchyTable(HierarchyTable);
+	TableHandler->ConstructHierarchy();
 
 	return HierarchyTable;
 }
@@ -47,7 +46,7 @@ bool UHierarchyTableFactory::ConfigureProperties()
 	// Allow the table type to configure itself
 	{
 		FHierarchyTableEditorModule& HierarchyTableModule = FModuleManager::GetModuleChecked<FHierarchyTableEditorModule>("HierarchyTableEditor");
-		const UHierarchyTable_TableTypeHandler* TableHandler = HierarchyTableModule.FindTableHandler(TableMetadata.GetScriptStruct());
+		const UHierarchyTable_TableTypeHandler* TableHandler = HierarchyTableModule.CreateTableHandler(TableMetadata.GetScriptStruct());
 		check(TableHandler);
 
 		bSuccess = TableHandler->FactoryConfigureProperties(TableMetadata);
diff --git a/Engine/Plugins/Experimental/Animation/HierarchyTable/Source/Editor/Private/HierarchyTableLabelWidget.cpp b/Engine/Plugins/Experimental/Animation/HierarchyTable/Source/Editor/Private/HierarchyTableLabelWidget.cpp
--- a/Engine/Plugins/Experimental/Animation/HierarchyTable/Source/Editor/Private/HierarchyTableLabelWidget.cpp
+++ b/Engine/Plugins/Experimental/Animation/HierarchyTable/Source/Editor/Private/HierarchyTableLabelWidget.cpp
@@ -48,10 +48,10 @@ TSharedPtr<SWidget> FHierarchyTableLabelWidgetConstructor::CreateWidget(
 		FTypedElementOverrideColumn* OverrideColumn = DataStorage->GetColumn<FTypedElementOverrideColumn>(TargetRow);
 
 		FHierarchyTableEditorModule& HierarchyTableModule = FModuleManager::GetModuleChecked<FHierarchyTableEditorModule>("HierarchyTableEditor");
-		const UHierarchyTable_TableTypeHandler* TableHandler = HierarchyTableModule.FindTableHandler(OverrideColumn->OwnerTable->TableMetadata.GetScriptStruct());
+		const TObjectPtr<UHierarchyTable_TableTypeHandler> TableHandler = HierarchyTableModule.CreateTableHandler(OverrideColumn->OwnerTable);
 		check(TableHandler);
 
-		return TableHandler->CreateLabelWidget(OverrideColumn->OwnerTable, OverrideColumn->OwnerEntryIndex);
+		return TableHandler->CreateLabelWidget(OverrideColumn->OwnerEntryIndex);
 	}
 	else
 	{
diff --git a/Engine/Plugins/Experimental/Animation/HierarchyTable/Source/Editor/Private/HierarchyTableTypeHandler.cpp b/Engine/Plugins/Experimental/Animation/HierarchyTable/Source/Editor/Private/HierarchyTableTypeHandler.cpp
--- a/Engine/Plugins/Experimental/Animation/HierarchyTable/Source/Editor/Private/HierarchyTableTypeHandler.cpp
+++ b/Engine/Plugins/Experimental/Animation/HierarchyTable/Source/Editor/Private/HierarchyTableTypeHandler.cpp
@@ -7,7 +7,7 @@
 #include "Widgets/Text/STextBlock.h"
 #include "Widgets/SBoxPanel.h"
 
-TSharedRef<SWidget> UHierarchyTable_TableTypeHandler::CreateLabelWidget(TObjectPtr<UHierarchyTable> HierarchyTable, const int32 EntryIndex) const
+TSharedRef<SWidget> UHierarchyTable_TableTypeHandler::CreateLabelWidget(const int32 EntryIndex) const
 {
 	return SNew(SHorizontalBox)
 		+SHorizontalBox::Slot()
diff --git a/Engine/Plugins/Experimental/Animation/HierarchyTable/Source/Editor/Public/HierarchyTableEditorModule.h b/Engine/Plugins/Experimental/Animation/HierarchyTable/Source/Editor/Public/HierarchyTableEditorModule.h
--- a/Engine/Plugins/Experimental/Animation/HierarchyTable/Source/Editor/Public/HierarchyTableEditorModule.h
+++ b/Engine/Plugins/Experimental/Animation/HierarchyTable/Source/Editor/Public/HierarchyTableEditorModule.h
@@ -5,6 +5,7 @@
 #include "HierarchyTableTypeHandler.h"
 #include "HierarchyTableType.h"
 #include "Modules/ModuleInterface.h"
+#include "HierarchyTable.h"
 
 class FHierarchyTableEditorModule : public IModuleInterface
 {
@@ -13,20 +14,22 @@ public:
 	
 	virtual void ShutdownModule() override;
 
-	HIERARCHYTABLEEDITOR_API void RegisterTableType(const UScriptStruct* TableType, const UHierarchyTable_TableTypeHandler* Handler);
+	HIERARCHYTABLEEDITOR_API void RegisterTableType(const UScriptStruct* TableType, const UClass* Handler);
 
 	HIERARCHYTABLEEDITOR_API void UnregisterTableType(const UScriptStruct* TableType);
 
-	HIERARCHYTABLEEDITOR_API const UHierarchyTable_TableTypeHandler*  FindTableHandler(const UScriptStruct* TableType) const;
+	HIERARCHYTABLEEDITOR_API const TObjectPtr<UHierarchyTable_TableTypeHandler> CreateTableHandler(const TObjectPtr<UHierarchyTable> HierarchyTable);
 
-	HIERARCHYTABLEEDITOR_API void RegisterElementType(const UScriptStruct* ElementType, const UHierarchyTable_ElementTypeHandler* Handler);
+	HIERARCHYTABLEEDITOR_API const TObjectPtr<UHierarchyTable_TableTypeHandler> CreateTableHandler(const TObjectPtr<const UScriptStruct> TableType);
 
-	HIERARCHYTABLEEDITOR_API void UnregisterElementType(const UScriptStruct* ElementType);
+	HIERARCHYTABLEEDITOR_API void RegisterElementTypeEditorColumns(const UScriptStruct* ElementType, const TArray<TObjectPtr<UScriptStruct>> Columns);
 
-	HIERARCHYTABLEEDITOR_API const UHierarchyTable_ElementTypeHandler* FindElementHandler(const UScriptStruct* ElementType) const;
+	HIERARCHYTABLEEDITOR_API void UnregisterElementTypeEditorColumns(const UScriptStruct* ElementType);
+
+	TArray<UScriptStruct*> GetElementTypeEditorColumns(const TObjectPtr<const UHierarchyTable> HierarchyTable);
 
 private:
-	TMap<const UScriptStruct*, const UHierarchyTable_TableTypeHandler*> TableHandlers;
+	TMap<const UScriptStruct*, const UClass*> TableHandlers;
 
-	TMap<const UScriptStruct*, const UHierarchyTable_ElementTypeHandler*> ElementHandlers;
+	TMap<const UScriptStruct*, TArray<UScriptStruct*>> EditorColumns;
 };
\ No newline at end of file
diff --git a/Engine/Plugins/Experimental/Animation/HierarchyTable/Source/Editor/Public/HierarchyTableTypeHandler.h b/Engine/Plugins/Experimental/Animation/HierarchyTable/Source/Editor/Public/HierarchyTableTypeHandler.h
--- a/Engine/Plugins/Experimental/Animation/HierarchyTable/Source/Editor/Public/HierarchyTableTypeHandler.h
+++ b/Engine/Plugins/Experimental/Animation/HierarchyTable/Source/Editor/Public/HierarchyTableTypeHandler.h
@@ -19,22 +19,18 @@ class HIERARCHYTABLEEDITOR_API UHierarchyTable_TableTypeHandler : public UObject
 	GENERATED_BODY()
 
 public:
-	virtual void InitialiseTable(const FInstancedStruct& TableType, TObjectPtr<UHierarchyTable> HierarchyTable) const {}
+	void SetHierarchyTable(TObjectPtr<UHierarchyTable> InHierarchyTable) { HierarchyTable = InHierarchyTable; }
 
-	virtual bool FactoryConfigureProperties(FInstancedStruct& TableType) const { return false; }
+	virtual void InitialiseTable(const FInstancedStruct& TableType) const {}
 
-	virtual void ExtendToolbar(UToolMenu* ToolMenu, UHierarchyTable* HierarchyTable, const IHierarchyTableEditor* Editor) const {}
+	virtual bool FactoryConfigureProperties(FInstancedStruct& TableType) const { return false; }
 
-	virtual TSharedRef<SWidget> CreateLabelWidget(TObjectPtr<UHierarchyTable> HierarchyTable, const int32 EntryIndex) const;
-};
+	virtual void ExtendToolbar(UToolMenu* ToolMenu, const IHierarchyTableEditor* Editor) const {}
 
-UCLASS(Abstract)
-class HIERARCHYTABLEEDITOR_API UHierarchyTable_ElementTypeHandler : public UObject
-{
-	GENERATED_BODY()
+	virtual TSharedRef<SWidget> CreateLabelWidget(const int32 EntryIndex) const;
 
-public:
-	virtual FInstancedStruct GetDefaultEntry() const { return FInstancedStruct(); }
+	virtual void ConstructHierarchy() {};
 
-	virtual TArray<UScriptStruct*> GetColumns() const { return {}; }
+protected:
+	TObjectPtr<UHierarchyTable> HierarchyTable;
 };
\ No newline at end of file
diff --git a/Engine/Plugins/Experimental/Animation/HierarchyTable/Source/Runtime/Private/HierarchyTable.cpp b/Engine/Plugins/Experimental/Animation/HierarchyTable/Source/Runtime/Private/HierarchyTable.cpp
--- a/Engine/Plugins/Experimental/Animation/HierarchyTable/Source/Runtime/Private/HierarchyTable.cpp
+++ b/Engine/Plugins/Experimental/Animation/HierarchyTable/Source/Runtime/Private/HierarchyTable.cpp
@@ -115,6 +115,14 @@ bool UHierarchyTable::UnregisterOnEntryAdded(const FDelegateHandle& Handle)
 	return OnEntryAdded.Remove(Handle);
 }
 
+FInstancedStruct UHierarchyTable::CreateDefaultValue() const
+{
+	FInstancedStruct OutStruct;
+	OutStruct.InitializeAs(ElementType);
+
+	return OutStruct;
+}
+
 const FHierarchyTableEntryData* const UHierarchyTable::GetTableEntry(const int32 EntryIndex) const
 {
 	if (TableData.IsValidIndex(EntryIndex))
diff --git a/Engine/Plugins/Experimental/Animation/HierarchyTable/Source/Runtime/Public/HierarchyTable.h b/Engine/Plugins/Experimental/Animation/HierarchyTable/Source/Runtime/Public/HierarchyTable.h
--- a/Engine/Plugins/Experimental/Animation/HierarchyTable/Source/Runtime/Public/HierarchyTable.h
+++ b/Engine/Plugins/Experimental/Animation/HierarchyTable/Source/Runtime/Public/HierarchyTable.h
@@ -97,6 +97,7 @@ class HIERARCHYTABLERUNTIME_API UHierarchyTable : public UObject
 public:
 	UHierarchyTable();
 
+private:
 	// The table metadata stores any data dependent on the table type needed to create and maintain the hierarchy.
 	UPROPERTY(VisibleAnywhere, Category = Default)
 	FInstancedStruct TableMetadata;
@@ -105,12 +106,48 @@ public:
 	UPROPERTY(VisibleAnywhere, Category = Default)
 	TObjectPtr<const UScriptStruct> ElementType;
 
-private:
 	// The actual table data
 	UPROPERTY()
 	TArray<FHierarchyTableEntryData> TableData;
 
 public:
+	template <typename InTableType>
+	bool IsTableType() const
+	{
+		return TableMetadata.GetScriptStruct() == InTableType::StaticStruct();
+	}
+
+	template <typename InTableType>
+	InTableType GetTableMetadata() const
+	{
+		check(IsTableType<InTableType>());
+		return TableMetadata.Get<InTableType>();
+	}
+
+	const UScriptStruct* GetTableMetadataStruct() const
+	{
+		return TableMetadata.GetScriptStruct();
+	}
+	
+	template <typename InElementType>
+	bool IsElementType() const
+	{
+		return ElementType == InElementType::StaticStruct();
+	}
+
+	FInstancedStruct CreateDefaultValue() const;
+
+	void Initialize(const FInstancedStruct& InTableMetadata, const TObjectPtr<const UScriptStruct> InElementType)
+	{
+		TableMetadata = InTableMetadata;
+		ElementType = InElementType;
+		TableData.Reset();
+	}
+
+	const FInstancedStruct& GetTableMetadata() const { return TableMetadata; }
+
+	const TObjectPtr<const UScriptStruct>& GetElementType() const { return ElementType; }
+
 	const TArray<FHierarchyTableEntryData>& GetTableData() const { return TableData; }
 
 	const FHierarchyTableEntryData* const GetTableEntry(const int32 EntryIndex) const;
diff --git a/Engine/Plugins/Experimental/Animation/HierarchyTableBuiltin/Source/Editor/Private/HierachyTableBuiltinEditorModule.cpp b/Engine/Plugins/Experimental/Animation/HierarchyTableBuiltin/Source/Editor/Private/HierachyTableBuiltinEditorModule.cpp
--- a/Engine/Plugins/Experimental/Animation/HierarchyTableBuiltin/Source/Editor/Private/HierachyTableBuiltinEditorModule.cpp
+++ b/Engine/Plugins/Experimental/Animation/HierarchyTableBuiltin/Source/Editor/Private/HierachyTableBuiltinEditorModule.cpp
@@ -2,21 +2,25 @@
 
 #include "HierarchyTableBuiltinEditorModule.h"
 #include "TimeProfile/HierarchyTableTypeTime.h"
-#include "TimeProfile/TimeProfileTypeHandler.h"
 #include "MaskProfile/HierarchyTableTypeMask.h"
-#include "MaskProfile/MaskProfileTypeHandler.h"
 #include "Modules/ModuleManager.h"
 #include "HierarchyTableEditorModule.h"
 #include "SkeletonHierarchyTableType.h"
 #include "SkeletonHierarchyTableTypeHandler.h"
+#include "MaskProfile/MaskProfileProxyColumn.h"
+#include "TimeProfile/TimeProfileProxyColumn.h"
 
 void FHierarchyTableBuiltinEditorModule::StartupModule()
 {
 	FHierarchyTableEditorModule& HierarchyTableModule = FModuleManager::LoadModuleChecked<FHierarchyTableEditorModule>("HierarchyTableEditor");
-	HierarchyTableModule.RegisterElementType(FHierarchyTable_ElementType_Mask::StaticStruct(), GetDefault<UHierarchyTable_ElementTypeHandler_Mask>());
-	HierarchyTableModule.RegisterElementType(FHierarchyTable_ElementType_Time::StaticStruct(), GetDefault<UHierarchyTable_ElementTypeHandler_Time>());
+	
+	HierarchyTableModule.RegisterTableType(FHierarchyTable_TableType_Skeleton::StaticStruct(), UHierarchyTable_TableTypeHandler_Skeleton::StaticClass());
 
-	HierarchyTableModule.RegisterTableType(FHierarchyTable_TableType_Skeleton::StaticStruct(), GetDefault<UHierarchyTable_TableTypeHandler_Skeleton>());
+	HierarchyTableModule.RegisterElementTypeEditorColumns(FHierarchyTable_ElementType_Mask::StaticStruct(),
+		{ FHierarchyTableMaskColumn_Value::StaticStruct() });
+
+	HierarchyTableModule.RegisterElementTypeEditorColumns(FHierarchyTable_ElementType_Time::StaticStruct(),
+		{ FHierarchyTableTimeColumn_StartTime::StaticStruct(), FHierarchyTableTimeColumn_EndTime::StaticStruct(), FHierarchyTableTimeColumn_TimeFactor::StaticStruct() });
 }
 
 void FHierarchyTableBuiltinEditorModule::ShutdownModule()
@@ -24,10 +28,9 @@ void FHierarchyTableBuiltinEditorModule::ShutdownModule()
 	if (FModuleManager::Get().IsModuleLoaded("HierarchyTableEditor"))
 	{
 		FHierarchyTableEditorModule& HierarchyTableModule = FModuleManager::GetModuleChecked<FHierarchyTableEditorModule>("HierarchyTableEditor");
-		HierarchyTableModule.UnregisterElementType(FHierarchyTable_ElementType_Mask::StaticStruct());
-		HierarchyTableModule.UnregisterElementType(FHierarchyTable_ElementType_Time::StaticStruct());
-
 		HierarchyTableModule.UnregisterTableType(FHierarchyTable_TableType_Skeleton::StaticStruct());
+		HierarchyTableModule.UnregisterElementTypeEditorColumns(FHierarchyTable_ElementType_Mask::StaticStruct());
+		HierarchyTableModule.UnregisterElementTypeEditorColumns(FHierarchyTable_ElementType_Time::StaticStruct());
 	}
 }
 
diff --git a/Engine/Plugins/Experimental/Animation/HierarchyTableBuiltin/Source/Editor/Private/MaskProfile/MaskProfileTypeHandler.cpp b/Engine/Plugins/Experimental/Animation/HierarchyTableBuiltin/Source/Editor/Private/MaskProfile/MaskProfileTypeHandler.cpp
deleted file mode 100644
diff --git a/Engine/Plugins/Experimental/Animation/HierarchyTableBuiltin/Source/Editor/Private/MaskProfile/MaskProfileTypeHandler.h b/Engine/Plugins/Experimental/Animation/HierarchyTableBuiltin/Source/Editor/Private/MaskProfile/MaskProfileTypeHandler.h
deleted file mode 100644
diff --git a/Engine/Plugins/Experimental/Animation/HierarchyTableBuiltin/Source/Editor/Private/SkeletonHierarchyTableTypeHandler.cpp b/Engine/Plugins/Experimental/Animation/HierarchyTableBuiltin/Source/Editor/Private/SkeletonHierarchyTableTypeHandler.cpp
--- a/Engine/Plugins/Experimental/Animation/HierarchyTableBuiltin/Source/Editor/Private/SkeletonHierarchyTableTypeHandler.cpp
+++ b/Engine/Plugins/Experimental/Animation/HierarchyTableBuiltin/Source/Editor/Private/SkeletonHierarchyTableTypeHandler.cpp
@@ -24,18 +24,14 @@
 
 #define LOCTEXT_NAMESPACE "UHierarchyTable_TableTypeHandler_Skeleton"
 
-void UHierarchyTable_TableTypeHandler_Skeleton::InitialiseTable(const FInstancedStruct& TableType, TObjectPtr<UHierarchyTable> HierarchyTable) const
+void UHierarchyTable_TableTypeHandler_Skeleton::ConstructHierarchy()
 {
-	const FHierarchyTable_TableType_Skeleton& SkeletonTableType = TableType.Get<FHierarchyTable_TableType_Skeleton>();
+	const FHierarchyTable_TableType_Skeleton& SkeletonTableType = HierarchyTable->GetTableMetadata<FHierarchyTable_TableType_Skeleton>();
 
 	check(SkeletonTableType.Skeleton);
 	check(HierarchyTable->GetTableData().Num() == 0);
 
-	FHierarchyTableEditorModule& HierarchyTableModule = FModuleManager::GetModuleChecked<FHierarchyTableEditorModule>("HierarchyTableEditor");
-	const UHierarchyTable_ElementTypeHandler* ElementHandler = HierarchyTableModule.FindElementHandler(HierarchyTable->ElementType);
-	check(ElementHandler);
-
-	FInstancedStruct DefaultEntry = ElementHandler->GetDefaultEntry();
+	FInstancedStruct DefaultEntry = HierarchyTable->CreateDefaultValue();
 
 	FInstancedStruct DefaultTablePayload;
 	DefaultTablePayload.InitializeAs<FHierarchyTable_TablePayloadType_Skeleton>();
@@ -91,7 +87,7 @@ bool UHierarchyTable_TableTypeHandler_Skeleton::FactoryConfigureProperties(FInst
 	return SkeletonTableType.Skeleton != nullptr;
 }
 
-void UHierarchyTable_TableTypeHandler_Skeleton::ExtendToolbar(UToolMenu* ToolMenu, UHierarchyTable* HierarchyTable, const IHierarchyTableEditor* Editor) const
+void UHierarchyTable_TableTypeHandler_Skeleton::ExtendToolbar(UToolMenu* ToolMenu, const IHierarchyTableEditor* Editor) const
 {
 	const FToolMenuInsert SectionInsertLocation("Asset", EToolMenuInsertType::After);
 	{
@@ -111,7 +107,7 @@ void UHierarchyTable_TableTypeHandler_Skeleton::ExtendToolbar(UToolMenu* ToolMen
 								{
 									const int32 SelectedIndex = Editor->GetSelectedEntryIndex();
 									const int32 ParentIndex = SelectedIndex == INDEX_NONE ? 0 : SelectedIndex;
-									AddCurve(HierarchyTable, ParentIndex, FName(CommittedText.ToString()));
+									AddCurve(ParentIndex, FName(CommittedText.ToString()));
 									FSlateApplication::Get().DismissAllMenus();
 								});
 
@@ -134,7 +130,7 @@ void UHierarchyTable_TableTypeHandler_Skeleton::ExtendToolbar(UToolMenu* ToolMen
 				
 				    InSubMenu->AddMenuEntry("AddNewCurve", Entry);
 
-					const FHierarchyTable_TableType_Skeleton& SkeletonTableMetadata = HierarchyTable->TableMetadata.Get<FHierarchyTable_TableType_Skeleton>();
+					const FHierarchyTable_TableType_Skeleton& SkeletonTableMetadata = HierarchyTable->GetTableMetadata<FHierarchyTable_TableType_Skeleton>();
 
 					FPersonaModule& PersonaModule = FModuleManager::LoadModuleChecked<FPersonaModule>("Persona");
 					InSubMenu->AddMenuEntry("ExistingCurveMenu", FToolMenuEntry::InitWidget(
@@ -148,7 +144,7 @@ void UHierarchyTable_TableTypeHandler_Skeleton::ExtendToolbar(UToolMenu* ToolMen
 									{
 										const int32 SelectedIndex = Editor->GetSelectedEntryIndex();
 										const int32 ParentIndex = SelectedIndex == INDEX_NONE ? 0 : SelectedIndex;
-										AddCurve(HierarchyTable, ParentIndex, InName);
+										AddCurve(ParentIndex, InName);
 										FSlateApplication::Get().DismissAllMenus();
 									}))
 						],
@@ -177,7 +173,7 @@ void UHierarchyTable_TableTypeHandler_Skeleton::ExtendToolbar(UToolMenu* ToolMen
 								{
 									const int32 SelectedIndex = Editor->GetSelectedEntryIndex();
 									const int32 ParentIndex = SelectedIndex == INDEX_NONE ? 0 : SelectedIndex;
-									AddAttribute(HierarchyTable, ParentIndex, FName(CommittedText.ToString()));
+									AddAttribute(ParentIndex, FName(CommittedText.ToString()));
 									FSlateApplication::Get().DismissAllMenus();
 								});
 
@@ -207,7 +203,7 @@ void UHierarchyTable_TableTypeHandler_Skeleton::ExtendToolbar(UToolMenu* ToolMen
 	}
 }
 
-TSharedRef<SWidget> UHierarchyTable_TableTypeHandler_Skeleton::CreateLabelWidget(TObjectPtr<UHierarchyTable> HierarchyTable, const int32 EntryIndex) const
+TSharedRef<SWidget> UHierarchyTable_TableTypeHandler_Skeleton::CreateLabelWidget(const int32 EntryIndex) const
 {
 	const FSlateBrush* LabelIcon = nullptr;
 	FSlateColor LabelIconColor = FSlateColor::UseForeground();
@@ -252,7 +248,7 @@ TSharedRef<SWidget> UHierarchyTable_TableTypeHandler_Skeleton::CreateLabelWidget
 		];
 }
 
-void UHierarchyTable_TableTypeHandler_Skeleton::AddCurve(TObjectPtr<UHierarchyTable> HierarchyTable, const int32 ParentIndex, const FName Identifier) const
+void UHierarchyTable_TableTypeHandler_Skeleton::AddCurve(const int32 ParentIndex, const FName Identifier) const
 {
 	FHierarchyTableEntryData NewEntry;
 	NewEntry.Parent = ParentIndex;
@@ -271,7 +267,7 @@ void UHierarchyTable_TableTypeHandler_Skeleton::AddCurve(TObjectPtr<UHierarchyTa
 	HierarchyTable->AddEntry(NewEntry);
 }
 	
-void UHierarchyTable_TableTypeHandler_Skeleton::AddAttribute(TObjectPtr<UHierarchyTable> HierarchyTable, const int32 ParentIndex, const FName Identifier) const
+void UHierarchyTable_TableTypeHandler_Skeleton::AddAttribute(const int32 ParentIndex, const FName Identifier) const
 {
 	FHierarchyTableEntryData NewEntry;
 	NewEntry.Parent = ParentIndex;
diff --git a/Engine/Plugins/Experimental/Animation/HierarchyTableBuiltin/Source/Editor/Private/SkeletonHierarchyTableTypeHandler.h b/Engine/Plugins/Experimental/Animation/HierarchyTableBuiltin/Source/Editor/Private/SkeletonHierarchyTableTypeHandler.h
--- a/Engine/Plugins/Experimental/Animation/HierarchyTableBuiltin/Source/Editor/Private/SkeletonHierarchyTableTypeHandler.h
+++ b/Engine/Plugins/Experimental/Animation/HierarchyTableBuiltin/Source/Editor/Private/SkeletonHierarchyTableTypeHandler.h
@@ -12,16 +12,16 @@ class UHierarchyTable_TableTypeHandler_Skeleton final : public UHierarchyTable_T
 	GENERATED_BODY()
 
 public:
-	virtual void InitialiseTable(const FInstancedStruct& TableType, TObjectPtr<UHierarchyTable> HierarchyTable) const;
+	virtual void ConstructHierarchy();
 
 	virtual bool FactoryConfigureProperties(FInstancedStruct& TableType) const override;
 
-	virtual void ExtendToolbar(UToolMenu* ToolMenu, UHierarchyTable* HierarchyTable, const IHierarchyTableEditor* Editor) const override;
+	virtual void ExtendToolbar(UToolMenu* ToolMenu, const IHierarchyTableEditor* Editor) const override;
 
-	virtual TSharedRef<SWidget> CreateLabelWidget(TObjectPtr<UHierarchyTable> HierarchyTable, const int32 EntryIndex) const;
+	virtual TSharedRef<SWidget> CreateLabelWidget(const int32 EntryIndex) const;
 
 private:
-	void AddCurve(TObjectPtr<UHierarchyTable> HierarchyTable, const int32 ParentIndex, const FName Identifier) const;
+	void AddCurve(const int32 ParentIndex, const FName Identifier) const;
 	
-	void AddAttribute(TObjectPtr<UHierarchyTable> HierarchyTable, const int32 ParentIndex, const FName Identifier) const;
+	void AddAttribute(const int32 ParentIndex, const FName Identifier) const;
 };
\ No newline at end of file
diff --git a/Engine/Plugins/Experimental/Animation/HierarchyTableBuiltin/Source/Editor/Private/TimeProfile/TimeProfileTypeHandler.cpp b/Engine/Plugins/Experimental/Animation/HierarchyTableBuiltin/Source/Editor/Private/TimeProfile/TimeProfileTypeHandler.cpp
deleted file mode 100644
diff --git a/Engine/Plugins/Experimental/Animation/HierarchyTableBuiltin/Source/Editor/Private/TimeProfile/TimeProfileTypeHandler.h b/Engine/Plugins/Experimental/Animation/HierarchyTableBuiltin/Source/Editor/Private/TimeProfile/TimeProfileTypeHandler.h
deleted file mode 100644
diff --git a/Engine/Plugins/Experimental/Animation/HierarchyTableBuiltin/Source/Runtime/Private/AnimNode_BlendProfileLayeredBlend.cpp b/Engine/Plugins/Experimental/Animation/HierarchyTableBuiltin/Source/Runtime/Private/AnimNode_BlendProfileLayeredBlend.cpp
--- a/Engine/Plugins/Experimental/Animation/HierarchyTableBuiltin/Source/Runtime/Private/AnimNode_BlendProfileLayeredBlend.cpp
+++ b/Engine/Plugins/Experimental/Animation/HierarchyTableBuiltin/Source/Runtime/Private/AnimNode_BlendProfileLayeredBlend.cpp
@@ -308,10 +308,10 @@ void FAnimNode_BlendProfileLayeredBlend::CreateMaskWeights(const USkeleton* Skel
 	PerBoneBlendWeights.Reset(NumBones);
 	PerBoneBlendWeights.AddZeroed(NumBones);
 
-	check(MaskTable->TableMetadata.GetScriptStruct() == FHierarchyTable_TableType_Skeleton::StaticStruct());
-	check(MaskTable->ElementType == FHierarchyTable_ElementType_Mask::StaticStruct());
+	check(MaskTable->IsTableType<FHierarchyTable_TableType_Skeleton>());
+	check(MaskTable->IsElementType<FHierarchyTable_ElementType_Mask>());
 
-	const FHierarchyTable_TableType_Skeleton& SkeletonMetadata = MaskTable->TableMetadata.Get<FHierarchyTable_TableType_Skeleton>();
+	const FHierarchyTable_TableType_Skeleton& SkeletonMetadata = MaskTable->GetTableMetadata<FHierarchyTable_TableType_Skeleton>();
 
 	for (int32 EntryIndex = 0; EntryIndex < MaskTable->GetTableData().Num(); ++EntryIndex)
 	{
diff --git a/Engine/Plugins/Experimental/Animation/HierarchyTableBuiltin/Source/Runtime/Public/MaskProfile/HierarchyTableTypeMask.h b/Engine/Plugins/Experimental/Animation/HierarchyTableBuiltin/Source/Runtime/Public/MaskProfile/HierarchyTableTypeMask.h
--- a/Engine/Plugins/Experimental/Animation/HierarchyTableBuiltin/Source/Runtime/Public/MaskProfile/HierarchyTableTypeMask.h
+++ b/Engine/Plugins/Experimental/Animation/HierarchyTableBuiltin/Source/Runtime/Public/MaskProfile/HierarchyTableTypeMask.h
@@ -12,7 +12,7 @@ struct FHierarchyTable_ElementType_Mask final : public FHierarchyTable_ElementTy
 	GENERATED_BODY()
 
 	FHierarchyTable_ElementType_Mask()
-		: Value(0.0f)
+		: Value(1.0f)
 	{
 	}
 
diff --git a/Engine/Plugins/Experimental/Animation/HierarchyTableBuiltin/Source/UncookedOnly/Private/AnimGraphNode_BlendProfileLayeredBlend.cpp b/Engine/Plugins/Experimental/Animation/HierarchyTableBuiltin/Source/UncookedOnly/Private/AnimGraphNode_BlendProfileLayeredBlend.cpp
--- a/Engine/Plugins/Experimental/Animation/HierarchyTableBuiltin/Source/UncookedOnly/Private/AnimGraphNode_BlendProfileLayeredBlend.cpp
+++ b/Engine/Plugins/Experimental/Animation/HierarchyTableBuiltin/Source/UncookedOnly/Private/AnimGraphNode_BlendProfileLayeredBlend.cpp
@@ -64,7 +64,7 @@ void UAnimGraphNode_BlendProfileLayeredBlend::ValidateAnimNodeDuringCompilation(
 		MessageLog.Error(*LOCTEXT("InvalidMask", "@@ has a null mask table.").ToString(), this, MaskTable);
 		bCompilationError = true;
 	}
-	else if (MaskTable->ElementType != FHierarchyTable_ElementType_Mask::StaticStruct())
+	else if (MaskTable->IsElementType<FHierarchyTable_ElementType_Mask>())
 	{
 		MessageLog.Error(*LOCTEXT("InvalidMaskType", "@@ has a null mask table that is not a mask type.").ToString(), this, MaskTable);
 		bCompilationError = true;
